- support update reference pose of skeletalmesh in engine

diff --git a/Engine/Source/Editor/SkeletalMeshEditor/Private/SkeletalMeshEditor.cpp b/Engine/Source/Editor/SkeletalMeshEditor/Private/SkeletalMeshEditor.cpp
--- a/Engine/Source/Editor/SkeletalMeshEditor/Private/SkeletalMeshEditor.cpp
+++ b/Engine/Source/Editor/SkeletalMeshEditor/Private/SkeletalMeshEditor.cpp
@@ -41,6 +41,7 @@
 #include "LODUtilities.h"
 #include "ScopedTransaction.h"
 #include "ComponentReregisterContext.h"
+#include "AnimPreviewInstance.h"
 
 const FName SkeletalMeshEditorAppIdentifier = FName(TEXT("SkeletalMeshEditorApp"));
 
@@ -174,6 +175,9 @@ void FSkeletalMeshEditor::BindCommands()
 
 	ToolkitCommands->MapAction(FPersonaCommonCommands::Get().TogglePlay,
 		FExecuteAction::CreateRaw(&GetPersonaToolkit()->GetPreviewScene().Get(), &IPersonaPreviewScene::TogglePlayback));
+
+	ToolkitCommands->MapAction(FSkeletalMeshEditorCommands::Get().UpdateRefPose,
+		FExecuteAction::CreateSP(this, &FSkeletalMeshEditor::HandleUpdateRefPose));
 }
 
 void FSkeletalMeshEditor::ExtendToolbar()
@@ -661,6 +665,25 @@ void FSkeletalMeshEditor::ExtendMenu()
 {
 	MenuExtender = MakeShareable(new FExtender);
 
+	struct Local
+	{
+		static void AddAssetMenu(FMenuBuilder& MenuBuilder, FSkeletalMeshEditor* InSkeletalMeshEditor)
+		{
+			MenuBuilder.BeginSection("SkeletalMeshEditor", LOCTEXT("SkeletalMeshEditorAssetMenu_Animation", "SkeletalMesh"));
+			{
+				MenuBuilder.AddMenuEntry(FSkeletalMeshEditorCommands::Get().UpdateRefPose);
+			}
+			MenuBuilder.EndSection();
+		}
+	};
+
+	MenuExtender->AddMenuExtension(
+		"AssetEditorActions",
+		EExtensionHook::After,
+		GetToolkitCommands(),
+		FMenuExtensionDelegate::CreateStatic(&Local::AddAssetMenu, this)
+	);
+
 	AddMenuExtender(MenuExtender);
 
 	ISkeletalMeshEditorModule& SkeletalMeshEditorModule = FModuleManager::GetModuleChecked<ISkeletalMeshEditorModule>("SkeletalMeshEditor");
@@ -738,6 +761,56 @@ void FSkeletalMeshEditor::HandleReimportMesh()
 	}
 }
 
+void FSkeletalMeshEditor::HandleUpdateRefPose()
+{
+	UDebugSkelMeshComponent* Component = GetPersonaToolkit()->GetPreviewMeshComponent();
+
+	if (SkeletalMesh && Component)
+	{
+		FScopedTransaction Transaction(LOCTEXT("ModifyRefPose", "Update Reference Pose"));
+		SkeletalMesh->Modify();
+
+		// need to be in the scope so that it SkeletonModifier to update the transform of RefSkeleton
+		{
+			USkeleton* Skeleton = SkeletalMesh->Skeleton;
+			FReferenceSkeleton& CurrentRefSkeleton = SkeletalMesh->RefSkeleton;
+			FReferenceSkeletonModifier RefSkeletonModifier(CurrentRefSkeleton, Skeleton);
+			// get current preview mesh transform
+			TArray<FTransform> CurrentPose = Component->GetComponentSpaceTransforms();
+
+			// we want to only apply these to raw ref pose transform
+			const TArray<FTransform>& RefPose = CurrentRefSkeleton.GetRawRefBonePose();
+			// make sure ref pose count is smaller than current pose
+			check(RefPose.Num() <= CurrentPose.Num());
+
+			for (int32 Index = 0; Index < RefPose.Num(); ++Index)
+			{
+				FTransform LocalTransform;
+				if (Index == 0)
+				{
+					LocalTransform = CurrentPose[Index];
+				}
+				else
+				{
+					const int32 ParentIndex = CurrentRefSkeleton.GetParentIndex(Index);
+					LocalTransform = CurrentPose[Index].GetRelativeTransform(CurrentPose[ParentIndex]);
+				}
+
+				RefSkeletonModifier.UpdateRefPoseTransform(Index, LocalTransform);
+			}
+		}
+
+		// calculate inverse ref matrices
+		SkeletalMesh->CalculateInvRefMatrices();
+
+		// now clear any bone transform modified
+		if (Component && Component->PreviewInstance)
+		{
+			Component->PreviewInstance->ResetModifiedBone();
+		}
+	}
+}
+
 void FSkeletalMeshEditor::HandleReimportAllMesh()
 {
 	// Reimport the asset
diff --git a/Engine/Source/Editor/SkeletalMeshEditor/Private/SkeletalMeshEditor.h b/Engine/Source/Editor/SkeletalMeshEditor/Private/SkeletalMeshEditor.h
--- a/Engine/Source/Editor/SkeletalMeshEditor/Private/SkeletalMeshEditor.h
+++ b/Engine/Source/Editor/SkeletalMeshEditor/Private/SkeletalMeshEditor.h
@@ -98,6 +98,7 @@ private:
 
 	void HandleReimportMesh();
 	void HandleReimportAllMesh();
+	void HandleUpdateRefPose();
 
 	/** Callback for toggling UV drawing in the viewport */
 	void ToggleMeshSectionSelection();
diff --git a/Engine/Source/Editor/SkeletalMeshEditor/Private/SkeletalMeshEditorCommands.cpp b/Engine/Source/Editor/SkeletalMeshEditor/Private/SkeletalMeshEditorCommands.cpp
--- a/Engine/Source/Editor/SkeletalMeshEditor/Private/SkeletalMeshEditorCommands.cpp
+++ b/Engine/Source/Editor/SkeletalMeshEditor/Private/SkeletalMeshEditorCommands.cpp
@@ -11,6 +11,8 @@ void FSkeletalMeshEditorCommands::RegisterCommands()
 	UI_COMMAND(ReimportAllMesh, "Reimport All Mesh", "Reimport the current mesh and all the custom LODs.", EUserInterfaceActionType::Button, FInputChord());
 
 	UI_COMMAND(MeshSectionSelection, "Section Selection", "Enables selecting Mesh Sections in the viewport (disables selecting bones using their physics shape).", EUserInterfaceActionType::ToggleButton, FInputChord());
+
+	UI_COMMAND(UpdateRefPose, "Update Reference Pose", "Update Reference Pose with current pose", EUserInterfaceActionType::Button, FInputChord());
 }
 
 #undef LOCTEXT_NAMESPACE
diff --git a/Engine/Source/Editor/SkeletalMeshEditor/Private/SkeletalMeshEditorCommands.h b/Engine/Source/Editor/SkeletalMeshEditor/Private/SkeletalMeshEditorCommands.h
--- a/Engine/Source/Editor/SkeletalMeshEditor/Private/SkeletalMeshEditorCommands.h
+++ b/Engine/Source/Editor/SkeletalMeshEditor/Private/SkeletalMeshEditorCommands.h
@@ -26,4 +26,7 @@ public:
 
 	// selecting mesh section using hit proxies
 	TSharedPtr<FUICommandInfo> MeshSectionSelection;
+
+	// update ref pose
+	TSharedPtr<FUICommandInfo> UpdateRefPose;
 };
diff --git a/Engine/Source/Editor/SkeletalMeshEditor/SkeletalMeshEditor.Build.cs b/Engine/Source/Editor/SkeletalMeshEditor/SkeletalMeshEditor.Build.cs
--- a/Engine/Source/Editor/SkeletalMeshEditor/SkeletalMeshEditor.Build.cs
+++ b/Engine/Source/Editor/SkeletalMeshEditor/SkeletalMeshEditor.Build.cs
@@ -22,7 +22,6 @@ public class SkeletalMeshEditor : ModuleRules
 				"SlateCore",
                 "EditorStyle",
                 "UnrealEd",
-                
                 "SkeletonEditor",
                 "Kismet",
                 "KismetWidgets",
@@ -34,7 +33,8 @@ public class SkeletalMeshEditor : ModuleRules
                 "RHI",
                 "ClothingSystemRuntime",
                 "ClothingSystemEditorInterface",
-				"ClothingSystemRuntimeInterface"
+				"ClothingSystemRuntimeInterface",
+				"AnimGraph"
             }
 		);
 
