Bone selection is now cleared when entering section selection mode in skeletal mesh editor

diff --git a/Engine/Source/Editor/SkeletalMeshEditor/Private/SkeletalMeshEditor.cpp b/Engine/Source/Editor/SkeletalMeshEditor/Private/SkeletalMeshEditor.cpp
--- a/Engine/Source/Editor/SkeletalMeshEditor/Private/SkeletalMeshEditor.cpp
+++ b/Engine/Source/Editor/SkeletalMeshEditor/Private/SkeletalMeshEditor.cpp
@@ -1049,9 +1049,10 @@ void FSkeletalMeshEditor::HandleReimportAllMeshWithNewFile(int32 SourceFileIndex
 void FSkeletalMeshEditor::ToggleMeshSectionSelection()
 {
 	TSharedRef<IPersonaPreviewScene> PreviewScene = GetPersonaToolkit()->GetPreviewScene();
+	PreviewScene->DeselectAll();
 	bool bState = !PreviewScene->AllowMeshHitProxies();
 	GetMutableDefault<UPersonaOptions>()->bAllowMeshSectionSelection = bState;
-	PreviewScene->SetAllowMeshHitProxies(bState);	
+	PreviewScene->SetAllowMeshHitProxies(bState);
 }
 
 bool FSkeletalMeshEditor::IsMeshSectionSelectionChecked() const
