- backing out the update ref pose option as that causes more issues than solving issues

diff --git a/Engine/Source/Editor/SkeletalMeshEditor/Private/SkeletalMeshEditor.cpp b/Engine/Source/Editor/SkeletalMeshEditor/Private/SkeletalMeshEditor.cpp
--- a/Engine/Source/Editor/SkeletalMeshEditor/Private/SkeletalMeshEditor.cpp
+++ b/Engine/Source/Editor/SkeletalMeshEditor/Private/SkeletalMeshEditor.cpp
@@ -43,7 +43,6 @@
 #include "ComponentReregisterContext.h"
 #include "EditorFramework/AssetImportData.h"
 #include "Factories/FbxSkeletalMeshImportData.h"
-#include "AnimPreviewInstance.h"
 
 const FName SkeletalMeshEditorAppIdentifier = FName(TEXT("SkeletalMeshEditorApp"));
 
@@ -177,9 +176,6 @@ void FSkeletalMeshEditor::BindCommands()
 
 	ToolkitCommands->MapAction(FPersonaCommonCommands::Get().TogglePlay,
 		FExecuteAction::CreateRaw(&GetPersonaToolkit()->GetPreviewScene().Get(), &IPersonaPreviewScene::TogglePlayback));
-
-	ToolkitCommands->MapAction(FSkeletalMeshEditorCommands::Get().UpdateRefPose,
-		FExecuteAction::CreateSP(this, &FSkeletalMeshEditor::HandleUpdateRefPose));
 }
 
 void FSkeletalMeshEditor::ExtendToolbar()
@@ -816,25 +812,6 @@ void FSkeletalMeshEditor::ExtendMenu()
 {
 	MenuExtender = MakeShareable(new FExtender);
 
-	struct Local
-	{
-		static void AddAssetMenu(FMenuBuilder& MenuBuilder, FSkeletalMeshEditor* InSkeletalMeshEditor)
-		{
-			MenuBuilder.BeginSection("SkeletalMeshEditor", LOCTEXT("SkeletalMeshEditorAssetMenu_Animation", "SkeletalMesh"));
-			{
-				MenuBuilder.AddMenuEntry(FSkeletalMeshEditorCommands::Get().UpdateRefPose);
-			}
-			MenuBuilder.EndSection();
-		}
-	};
-
-	MenuExtender->AddMenuExtension(
-		"AssetEditorActions",
-		EExtensionHook::After,
-		GetToolkitCommands(),
-		FMenuExtensionDelegate::CreateStatic(&Local::AddAssetMenu, this)
-	);
-
 	AddMenuExtender(MenuExtender);
 
 	ISkeletalMeshEditorModule& SkeletalMeshEditorModule = FModuleManager::GetModuleChecked<ISkeletalMeshEditorModule>("SkeletalMeshEditor");
@@ -966,56 +943,6 @@ void ReimportAllCustomLODs(USkeletalMesh* SkeletalMesh, UDebugSkelMeshComponent*
 	}
 }
 
-void FSkeletalMeshEditor::HandleUpdateRefPose()
-{
-	UDebugSkelMeshComponent* Component = GetPersonaToolkit()->GetPreviewMeshComponent();
-
-	if (SkeletalMesh && Component)
-	{
-		FScopedTransaction Transaction(LOCTEXT("ModifyRefPose", "Update Reference Pose"));
-		SkeletalMesh->Modify();
-
-		// need to be in the scope so that it SkeletonModifier to update the transform of RefSkeleton
-		{
-			USkeleton* Skeleton = SkeletalMesh->Skeleton;
-			FReferenceSkeleton& CurrentRefSkeleton = SkeletalMesh->RefSkeleton;
-			FReferenceSkeletonModifier RefSkeletonModifier(CurrentRefSkeleton, Skeleton);
-			// get current preview mesh transform
-			TArray<FTransform> CurrentPose = Component->GetComponentSpaceTransforms();
-
-			// we want to only apply these to raw ref pose transform
-			const TArray<FTransform>& RefPose = CurrentRefSkeleton.GetRawRefBonePose();
-			// make sure ref pose count is smaller than current pose
-			check(RefPose.Num() <= CurrentPose.Num());
-
-			for (int32 Index = 0; Index < RefPose.Num(); ++Index)
-			{
-				FTransform LocalTransform;
-				if (Index == 0)
-				{
-					LocalTransform = CurrentPose[Index];
-				}
-				else
-				{
-					const int32 ParentIndex = CurrentRefSkeleton.GetParentIndex(Index);
-					LocalTransform = CurrentPose[Index].GetRelativeTransform(CurrentPose[ParentIndex]);
-				}
-
-				RefSkeletonModifier.UpdateRefPoseTransform(Index, LocalTransform);
-			}
-		}
-
-		// calculate inverse ref matrices
-		SkeletalMesh->CalculateInvRefMatrices();
-
-		// now clear any bone transform modified
-		if (Component && Component->PreviewInstance)
-		{
-			Component->PreviewInstance->ResetModifiedBone();
-		}
-	}
-}
-
 void FSkeletalMeshEditor::HandleReimportAllMesh(int32 SourceFileIndex /*= INDEX_NONE*/)
 {
 	// Reimport the asset
diff --git a/Engine/Source/Editor/SkeletalMeshEditor/Private/SkeletalMeshEditor.h b/Engine/Source/Editor/SkeletalMeshEditor/Private/SkeletalMeshEditor.h
--- a/Engine/Source/Editor/SkeletalMeshEditor/Private/SkeletalMeshEditor.h
+++ b/Engine/Source/Editor/SkeletalMeshEditor/Private/SkeletalMeshEditor.h
@@ -103,8 +103,6 @@ private:
 	void HandleReimportAllMesh(int32 SourceFileIndex = INDEX_NONE);
 	void HandleReimportAllMeshWithNewFile(int32 SourceFileIndex = INDEX_NONE);
 
-	void HandleUpdateRefPose();
-
 	/** Callback for toggling UV drawing in the viewport */
 	void ToggleMeshSectionSelection();
 
diff --git a/Engine/Source/Editor/SkeletalMeshEditor/Private/SkeletalMeshEditorCommands.cpp b/Engine/Source/Editor/SkeletalMeshEditor/Private/SkeletalMeshEditorCommands.cpp
--- a/Engine/Source/Editor/SkeletalMeshEditor/Private/SkeletalMeshEditorCommands.cpp
+++ b/Engine/Source/Editor/SkeletalMeshEditor/Private/SkeletalMeshEditorCommands.cpp
@@ -13,8 +13,6 @@ void FSkeletalMeshEditorCommands::RegisterCommands()
 	UI_COMMAND(ReimportAllMeshWithNewFile, "Reimport Base Mesh + LODs With New File", "Reimport the base mesh using a new source file and all the custom LODs (No new source file for LODs).", EUserInterfaceActionType::Button, FInputChord());
 
 	UI_COMMAND(MeshSectionSelection, "Section Selection", "Enables selecting Mesh Sections in the viewport (disables selecting bones using their physics shape).", EUserInterfaceActionType::ToggleButton, FInputChord());
-
-	UI_COMMAND(UpdateRefPose, "Update Reference Pose", "Update Reference Pose with current pose", EUserInterfaceActionType::Button, FInputChord());
 }
 
 #undef LOCTEXT_NAMESPACE
diff --git a/Engine/Source/Editor/SkeletalMeshEditor/Private/SkeletalMeshEditorCommands.h b/Engine/Source/Editor/SkeletalMeshEditor/Private/SkeletalMeshEditorCommands.h
--- a/Engine/Source/Editor/SkeletalMeshEditor/Private/SkeletalMeshEditorCommands.h
+++ b/Engine/Source/Editor/SkeletalMeshEditor/Private/SkeletalMeshEditorCommands.h
@@ -28,7 +28,4 @@ public:
 
 	// selecting mesh section using hit proxies
 	TSharedPtr<FUICommandInfo> MeshSectionSelection;
-
-	// update ref pose
-	TSharedPtr<FUICommandInfo> UpdateRefPose;
 };
diff --git a/Engine/Source/Editor/SkeletalMeshEditor/SkeletalMeshEditor.Build.cs b/Engine/Source/Editor/SkeletalMeshEditor/SkeletalMeshEditor.Build.cs
--- a/Engine/Source/Editor/SkeletalMeshEditor/SkeletalMeshEditor.Build.cs
+++ b/Engine/Source/Editor/SkeletalMeshEditor/SkeletalMeshEditor.Build.cs
@@ -22,6 +22,7 @@ public class SkeletalMeshEditor : ModuleRules
 				"SlateCore",
                 "EditorStyle",
                 "UnrealEd",
+                
                 "SkeletonEditor",
                 "Kismet",
                 "KismetWidgets",
@@ -33,8 +34,7 @@ public class SkeletalMeshEditor : ModuleRules
                 "RHI",
                 "ClothingSystemRuntime",
                 "ClothingSystemEditorInterface",
-				"ClothingSystemRuntimeInterface",
-				"AnimGraph"
+				"ClothingSystemRuntimeInterface"
             }
 		);
 
