Skeletal mesh editor now opens the skeleon tree by default

diff --git a/Engine/Source/Editor/SkeletalMeshEditor/Private/SkeletalMeshEditorMode.cpp b/Engine/Source/Editor/SkeletalMeshEditor/Private/SkeletalMeshEditorMode.cpp
--- a/Engine/Source/Editor/SkeletalMeshEditor/Private/SkeletalMeshEditorMode.cpp
+++ b/Engine/Source/Editor/SkeletalMeshEditor/Private/SkeletalMeshEditorMode.cpp
@@ -35,7 +35,7 @@ FSkeletalMeshEditorMode::FSkeletalMeshEditorMode(TSharedRef<FWorkflowCentricAppl
 
 	TabFactories.RegisterFactory(CreateMeshControllerMappingTabFactory(InHostingApp, Cast<USkeletalMesh> (SkeletalMeshEditor->HandleGetAsset()), SkeletalMeshEditor->OnPostUndo));
 
-	TabLayout = FTabManager::NewLayout("Standalone_SkeletalMeshEditor_Layout_v3.1")
+	TabLayout = FTabManager::NewLayout("Standalone_SkeletalMeshEditor_Layout_v3.2")
 		->AddArea
 		(
 			FTabManager::NewPrimaryArea()
@@ -57,8 +57,9 @@ FSkeletalMeshEditorMode::FSkeletalMeshEditorMode(TSharedRef<FWorkflowCentricAppl
 					FTabManager::NewStack()
 					->SetSizeCoefficient(0.2f)
 					->SetHideTabWell(false)
-					->AddTab(SkeletalMeshEditorTabs::SkeletonTreeTab, ETabState::ClosedTab)
 					->AddTab(SkeletalMeshEditorTabs::AssetDetailsTab, ETabState::OpenedTab)
+					->AddTab(SkeletalMeshEditorTabs::SkeletonTreeTab, ETabState::OpenedTab)
+					->SetForegroundTab(SkeletalMeshEditorTabs::AssetDetailsTab)
 				)
 				->Split
 				(
