Refresh skeleton tree when we reimport in the skeletal mesh editor

diff --git a/Engine/Source/Editor/SkeletalMeshEditor/Private/SkeletalMeshEditor.cpp b/Engine/Source/Editor/SkeletalMeshEditor/Private/SkeletalMeshEditor.cpp
--- a/Engine/Source/Editor/SkeletalMeshEditor/Private/SkeletalMeshEditor.cpp
+++ b/Engine/Source/Editor/SkeletalMeshEditor/Private/SkeletalMeshEditor.cpp
@@ -884,7 +884,12 @@ UObject* FSkeletalMeshEditor::HandleGetAsset()
 bool FSkeletalMeshEditor::HandleReimportMeshInternal(int32 SourceFileIndex /*= INDEX_NONE*/, bool bWithNewFile /*= false*/)
 {
 	// Reimport the asset
-	return FReimportManager::Instance()->Reimport(SkeletalMesh, true, true, TEXT(""), nullptr, SourceFileIndex, bWithNewFile);
+	const bool bResult = FReimportManager::Instance()->Reimport(SkeletalMesh, true, true, TEXT(""), nullptr, SourceFileIndex, bWithNewFile);
+
+	// Refresh skeleton tree
+	SkeletonTree->Refresh();
+
+	return bResult;
 }
 
 void FSkeletalMeshEditor::HandleReimportMesh(int32 SourceFileIndex /*= INDEX_NONE*/)
