AltSkinWeights - Crash - Reimporting Base Mesh while previewing Skin Weight Profile causes editor to crash
#fix Clear and reset profile when reimporting skeletal mesh

diff --git a/Engine/Source/Editor/SkeletalMeshEditor/Private/SkeletalMeshEditor.cpp b/Engine/Source/Editor/SkeletalMeshEditor/Private/SkeletalMeshEditor.cpp
--- a/Engine/Source/Editor/SkeletalMeshEditor/Private/SkeletalMeshEditor.cpp
+++ b/Engine/Source/Editor/SkeletalMeshEditor/Private/SkeletalMeshEditor.cpp
@@ -883,12 +883,30 @@ UObject* FSkeletalMeshEditor::HandleGetAsset()
 
 bool FSkeletalMeshEditor::HandleReimportMeshInternal(int32 SourceFileIndex /*= INDEX_NONE*/, bool bWithNewFile /*= false*/)
 {
+	// Make sure we clear any skin weight profile currently being previewed
+	UDebugSkelMeshComponent * MeshComponent = GetPersonaToolkit()->GetPreviewScene()->GetPreviewMeshComponent();
+	const FName ProfileName = MeshComponent->GetCurrentSkinWeightProfileName();
+	if (MeshComponent)
+	{
+		MeshComponent->ClearSkinWeightProfile();
+	}
+
 	// Reimport the asset
 	const bool bResult = FReimportManager::Instance()->Reimport(SkeletalMesh, true, true, TEXT(""), nullptr, SourceFileIndex, bWithNewFile);
 
 	// Refresh skeleton tree
 	SkeletonTree->Refresh();
 
+	if (ProfileName != NAME_None)
+	{
+		if (MeshComponent)
+		{
+			MeshComponent->SetSkinWeightProfile(ProfileName);
+		}
+	}
+
+	
+
 	return bResult;
 }
 
