Static analysis fix

diff --git a/Engine/Source/Editor/SkeletalMeshEditor/Private/SkeletalMeshEditor.cpp b/Engine/Source/Editor/SkeletalMeshEditor/Private/SkeletalMeshEditor.cpp
--- a/Engine/Source/Editor/SkeletalMeshEditor/Private/SkeletalMeshEditor.cpp
+++ b/Engine/Source/Editor/SkeletalMeshEditor/Private/SkeletalMeshEditor.cpp
@@ -884,8 +884,8 @@ UObject* FSkeletalMeshEditor::HandleGetAsset()
 bool FSkeletalMeshEditor::HandleReimportMeshInternal(int32 SourceFileIndex /*= INDEX_NONE*/, bool bWithNewFile /*= false*/)
 {
 	// Make sure we clear any skin weight profile currently being previewed
-	UDebugSkelMeshComponent * MeshComponent = GetPersonaToolkit()->GetPreviewScene()->GetPreviewMeshComponent();
-	const FName ProfileName = MeshComponent->GetCurrentSkinWeightProfileName();
+	UDebugSkelMeshComponent* MeshComponent = GetPersonaToolkit()->GetPreviewScene()->GetPreviewMeshComponent();
+	const FName ProfileName = MeshComponent ? MeshComponent->GetCurrentSkinWeightProfileName() : NAME_None;
 	if (MeshComponent)
 	{
 		MeshComponent->ClearSkinWeightProfile();
@@ -897,16 +897,12 @@ bool FSkeletalMeshEditor::HandleReimportMeshInternal(int32 SourceFileIndex /*= I
 	// Refresh skeleton tree
 	SkeletonTree->Refresh();
 
-	if (ProfileName != NAME_None)
+	// Re-set skin weigh profile
+	if (MeshComponent && (ProfileName != NAME_None))
 	{
-		if (MeshComponent)
-		{
-			MeshComponent->SetSkinWeightProfile(ProfileName);
-		}
+		MeshComponent->SetSkinWeightProfile(ProfileName);
 	}
 
-	
-
 	return bResult;
 }
 
