Copying //UE4/Dev-Editor to Dev-Main (//UE4/Dev-Main) for 4.24

diff --git a/Engine/Source/Editor/Persona/Private/PersonaMeshDetails.h b/Engine/Source/Editor/Persona/Private/PersonaMeshDetails.h
--- a/Engine/Source/Editor/Persona/Private/PersonaMeshDetails.h
+++ b/Engine/Source/Editor/Persona/Private/PersonaMeshDetails.h
@@ -14,6 +14,8 @@
 #include "IDetailCustomization.h"
 #include "Widgets/Input/SComboBox.h"
 
+DECLARE_LOG_CATEGORY_EXTERN(LogSkeletalMeshPersonaMeshDetail, Log, All);
+
 struct FAssetData;
 class FDetailWidgetRow;
 class FPersonaMeshDetails;
@@ -25,6 +27,10 @@ struct FSectionLocalizer;
 class IDetailCategoryBuilder;
 class IDetailGroup;
 
+//The parameter is the LOD index so we know which LOD need to be rebuild
+DECLARE_DELEGATE_RetVal_OneParam(bool, FIsLODSettingsEnabledDelegate, int32)
+DECLARE_DELEGATE_OneParam(FModifyMeshLODSettingsDelegate, int32)
+
 /**
  * Struct to uniquely identify clothing applied to a material section
  * Contains index into the ClothingAssets array and the submesh index.
@@ -80,6 +86,183 @@ struct FSectionLocalizer
 	int32 SectionIndex;
 };
 
+class FSkeletalMeshReductionSettingsLayout : public IDetailCustomNodeBuilder, public TSharedFromThis<FSkeletalMeshReductionSettingsLayout>
+{
+public:
+	FSkeletalMeshReductionSettingsLayout(FSkeletalMeshOptimizationSettings& InReductionettings, bool InbIsLODModelbuildDataAvailable, int32 InLODIndex, FIsLODSettingsEnabledDelegate InIsLODSettingsEnabledDelegate, FModifyMeshLODSettingsDelegate InModifyMeshLODSettingsDelegate);
+	virtual ~FSkeletalMeshReductionSettingsLayout() {};
+
+	enum EImportanceType
+	{
+		ID_Silhouette,
+		ID_Texture,
+		ID_Shading,
+		ID_Skinning
+	};
+
+	DECLARE_DELEGATE_RetVal(float, FGetFloatDelegate);
+	DECLARE_DELEGATE_OneParam(FSetFloatDelegate, float);
+
+	DECLARE_DELEGATE_RetVal(int32, FGetIntegerDelegate);
+	DECLARE_DELEGATE_OneParam(FSetIntegerDelegate, int32);
+
+	DECLARE_DELEGATE_RetVal(ECheckBoxState, FGetCheckBoxStateDelegate);
+	DECLARE_DELEGATE_OneParam(FSetCheckBoxStateDelegate, ECheckBoxState);
+private:
+	/** IDetailCustomNodeBuilder Interface*/
+	virtual void SetOnRebuildChildren(FSimpleDelegate InOnRegenerateChildren) override {}
+	virtual void GenerateHeaderRowContent(FDetailWidgetRow& NodeRow) override;
+	virtual void GenerateChildContent(IDetailChildrenBuilder& ChildrenBuilder) override;
+	virtual void Tick(float DeltaTime) override {}
+	virtual bool RequiresTick() const override { return false; }
+	virtual FName GetName() const override { static FName MeshBuildSettings("MeshBuildSettings"); return MeshBuildSettings; }
+	virtual bool InitiallyCollapsed() const override { return true; }
+
+	bool IsReductionEnabled() const;
+
+	//Custom Row Add utilities
+	FDetailWidgetRow& AddFloatRow(IDetailChildrenBuilder& ChildrenBuilder, const FText RowTitleText, const FText RowNameContentText, const FText RowNameContentTootlipText, const float MinSliderValue, const float MaxSliderValue, FGetFloatDelegate GetterDelegate, FSetFloatDelegate SetterDelegate);
+	FDetailWidgetRow& AddBoolRow(IDetailChildrenBuilder& ChildrenBuilder, const FText RowTitleText, const FText RowNameContentText, const FText RowNameContentToolitipText, FGetCheckBoxStateDelegate GetterDelegate, FSetCheckBoxStateDelegate SetterDelegate);
+	FDetailWidgetRow& AddIntegerRow(IDetailChildrenBuilder& ChildrenBuilder, const FText RowTitleText, const FText RowNameContentText, const FText RowNameContentTootlipText, const int32 MinSliderValue, const int32 MaxSliderValue, FGetIntegerDelegate GetterDelegate, FSetIntegerDelegate SetterDelegate);
+	void AddBaseLODRow(IDetailChildrenBuilder& ChildrenBuilder);
+
+	void SetPercentAndAbsoluteVisibility(FDetailWidgetRow& Row, SkeletalMeshTerminationCriterion FirstCriterion, SkeletalMeshTerminationCriterion SecondCriterion);
+
+	int32 GetBaseLODValue() const
+	{
+		return ReductionSettings.BaseLOD;
+	}
+	void SetBaseLODValue(int32 Value)
+	{
+		ReductionSettings.BaseLOD = Value;
+	}
+	TSharedRef<class SWidget> FillReductionMethodMenu();
+	FText GetReductionMethodText() const;
+
+	TSharedRef<class SWidget> FillReductionImportanceMenu(const EImportanceType Importance);
+	FText GetReductionImportanceText(const EImportanceType Importance) const;
+
+	TSharedRef<class SWidget> FillReductionTerminationCriterionMenu();
+	FText GetReductionTerminationCriterionText() const;
+
+	float GetNumTrianglesPercentage() const;
+	void SetNumTrianglesPercentage(float Value);
+
+	float GetNumVerticesPercentage() const;
+	void SetNumVerticesPercentage(float Value);
+
+	int32 GetNumMaxTrianglesCount() const;
+	void SetNumMaxTrianglesCount(int32 Value);
+
+	int32 GetNumMaxVerticesCount() const;
+	void SetNumMaxVerticesCount(int32 Value);
+
+	float GetAccuracyPercentage() const;
+	void SetAccuracyPercentage(float Value);
+
+	ECheckBoxState ShouldRecomputeNormals() const;
+	void OnRecomputeNormalsChanged(ECheckBoxState NewState);
+
+	float GetNormalsThreshold() const;
+	void SetNormalsThreshold(float Value);
+
+	float GetWeldingThreshold() const;
+	void SetWeldingThreshold(float Value);
+
+	ECheckBoxState GetLockEdges() const;
+	void SetLockEdges(ECheckBoxState NewState);
+
+	ECheckBoxState GetEnforceBoneBoundaries() const;
+	void SetEnforceBoneBoundaries(ECheckBoxState NewState);
+
+	float GetVolumeImportance() const;
+	void SetVolumeImportance(float Value);
+	
+	ECheckBoxState GetRemapMorphTargets() const;
+	void SetRemapMorphTargets(ECheckBoxState NewState);
+
+	int32 GetMaxBonesPerVertex() const;
+	void SetMaxBonesPerVertex(int32 Value);
+
+	// Used the the thrid-party UI.  
+	EVisibility GetVisibiltyIfCurrentReductionMethodIsNot(SkeletalMeshOptimizationType ReductionType) const;
+
+	// Used by the native tool UI.
+	EVisibility ShowIfCurrentCriterionIs(TArray<SkeletalMeshTerminationCriterion> TerminationCriterionArray) const;
+
+	/** Detect usage of thirdparty vs native tool */
+	bool UseNativeLODTool() const;
+	bool UseNativeReductionTool() const;
+
+	/**
+	Used to hide parameters that only make sense for the third party tool.
+	@return EVisibility::Visible if we are using the simplygon tool, otherwise EVisibility::Hidden
+	*/
+	EVisibility GetVisibilityForThirdPartyTool() const;
+
+private:
+	FSkeletalMeshOptimizationSettings& ReductionSettings;
+	bool bIsLODModelbuildDataAvailable;
+	int32 LODIndex;
+	FIsLODSettingsEnabledDelegate IsLODSettingsEnabledDelegate;
+	FModifyMeshLODSettingsDelegate ModifyMeshLODSettingsDelegate;
+
+	UEnum* EnumReductionMethod;
+	UEnum* EnumImportance;
+	UEnum* EnumTerminationCriterion;
+
+	//Use this data to keep a valid reference so the helper lambda can have persistent data
+	//Helper lambda are use with spinner to not do a transaction when spinning with mouse movement
+	struct FSliderStateData
+	{
+		float MovementValueFloat = 0.0f;
+		int32 MovementValueInt = 0;
+		bool bSliderActiveMode = false;
+	};
+	TArray<FSliderStateData> SliderStateDataArray;
+};
+
+class FSkeletalMeshBuildSettingsLayout : public IDetailCustomNodeBuilder, public TSharedFromThis<FSkeletalMeshBuildSettingsLayout>
+{
+public:
+	FSkeletalMeshBuildSettingsLayout(FSkeletalMeshBuildSettings& InBuildSettings, int32 InLODIndex, FIsLODSettingsEnabledDelegate InIsBuildSettingsEnabledDelegate, FModifyMeshLODSettingsDelegate InModifyMeshLODSettingsDelegate);
+	virtual ~FSkeletalMeshBuildSettingsLayout() {};
+
+private:
+	/** IDetailCustomNodeBuilder Interface*/
+	virtual void SetOnRebuildChildren(FSimpleDelegate InOnRegenerateChildren) override {}
+	virtual void GenerateHeaderRowContent(FDetailWidgetRow& NodeRow) override;
+	virtual void GenerateChildContent(IDetailChildrenBuilder& ChildrenBuilder) override;
+	virtual void Tick(float DeltaTime) override {}
+	virtual bool RequiresTick() const override { return false; }
+	virtual FName GetName() const override { static FName MeshBuildSettings("MeshBuildSettings"); return MeshBuildSettings; }
+	virtual bool InitiallyCollapsed() const override { return true; }
+
+	bool IsBuildEnabled() const;
+	ECheckBoxState ShouldRecomputeNormals() const;
+	ECheckBoxState ShouldRecomputeTangents() const;
+	ECheckBoxState ShouldUseMikkTSpace() const;
+	ECheckBoxState ShouldComputeWeightedNormals() const;
+	ECheckBoxState ShouldRemoveDegenerates() const;
+	ECheckBoxState ShouldUseHighPrecisionTangentBasis() const;
+	ECheckBoxState ShouldUseFullPrecisionUVs() const;
+	ECheckBoxState ShouldBuildAdjacencyBuffer() const;
+
+	void OnRecomputeNormalsChanged(ECheckBoxState NewState);
+	void OnRecomputeTangentsChanged(ECheckBoxState NewState);
+	void OnUseMikkTSpaceChanged(ECheckBoxState NewState);
+	void OnComputeWeightedNormalsChanged(ECheckBoxState NewState);
+	void OnRemoveDegeneratesChanged(ECheckBoxState NewState);
+	void OnUseHighPrecisionTangentBasisChanged(ECheckBoxState NewState);
+	void OnUseFullPrecisionUVsChanged(ECheckBoxState NewState);
+	void OnBuildAdjacencyBufferChanged(ECheckBoxState NewState);
+private:
+	FSkeletalMeshBuildSettings& BuildSettings;
+	int32 LODIndex;
+	FIsLODSettingsEnabledDelegate IsBuildSettingsEnabledDelegate;
+	FModifyMeshLODSettingsDelegate ModifyMeshLODSettingsDelegate;
+};
+
 class FPersonaMeshDetails : public IDetailCustomization
 {
 public:
@@ -93,6 +276,8 @@ public:
 	virtual void CustomizeDetails( IDetailLayoutBuilder& DetailLayout ) override;
 
 private:
+	//This function customize the LODInfo temporary object
+	void CustomizeLODInfoSetingsDetails(IDetailLayoutBuilder& DetailLayout, class ULODInfoUILayout* LODInfoUILayout, TSharedRef<IPropertyHandle> LODInfoProperty, IDetailCategoryBuilder& LODCategory);
 
 	FReply AddMaterialSlot();
 
@@ -265,34 +450,6 @@ private:
 	*/
 	void OnSectionIsolatedChanged(ECheckBoxState NewState, int32 SectionIndex);
 
-	/**
-	 * Handler for check box display based on whether the material has shadow casting enabled
-	 *
-	 * @param MaterialIndex	The material index which a section in a specific LOD model has
-	 */
-	ECheckBoxState IsShadowCastingEnabled(int32 MaterialIndex) const;
-
-	/**
-	 * Handler for changing shadow casting status on a material
-	 *
-	 * @param MaterialIndex	The material index which a section in a specific LOD model has
-	 */
-	void OnShadowCastingChanged(ECheckBoxState NewState, int32 MaterialIndex);
-
-	/**
-	* Handler for check box display based on whether this section does recalculate normal or not
-	*
-	* @param MaterialIndex	The material index which a section in a specific LOD model has
-	*/
-	ECheckBoxState IsRecomputeTangentEnabled(int32 MaterialIndex) const;
-
-	/**
-	* Handler for changing recalulate normal status on a material
-	*
-	* @param MaterialIndex	The material index which a section in a specific LOD model has
-	*/
-	void OnRecomputeTangentChanged(ECheckBoxState NewState, int32 MaterialIndex);
-
 	/**
 	* Handler for check box display based on whether the material has shadow casting enabled
 	*
@@ -366,15 +523,15 @@ private:
 	/** apply LOD changes if the user modified LOD reduction settings */
 	FReply OnApplyChanges();
 	/** regenerate one specific LOD Index no dependencies*/
-	void RegenerateOneLOD(int32 LODIndex, bool bReregisterComponent = true);
+	void RegenerateOneLOD(int32 LODIndex);
 	/** regenerate the specific all LODs dependent of InLODIndex. This is not regenerating the InLODIndex*/
-	void RegenerateDependentLODs(int32 LODIndex, bool bReregisterComponent = true);
+	void RegenerateDependentLODs(int32 LODIndex);
+	/** Apply specified LOD Index */
+	FReply ApplyLODChanges(int32 LODIndex);
 	/** Apply specified LOD Index */
 	FReply RegenerateLOD(int32 LODIndex);
 	/** Removes the specified lod from the skeletal mesh */
 	FReply RemoveOneLOD(int32 LODIndex);
-	/** Remove Bones again */
-	FReply RemoveBones(int32 LODIndex);
 	/** hide properties which don't need to be showed to end users */
 	void HideUnnecessaryProperties(IDetailLayoutBuilder& DetailLayout);
 
@@ -454,6 +611,10 @@ private:
 
 	IDetailLayoutBuilder* MeshDetailLayout;
 
+	//This is the mockup UObjects to modify a copy of the LODInfo
+	TArray<class ULODInfoUILayout*> LODInfoUILayouts;
+	TArray<TSharedRef<IDetailsView>> LODInfoUILayoutDetailsViews;
+
 	/** LOD import options */
 	TArray<TSharedPtr<FString> > LODNames;
 	/** Helper value that corresponds to the 'Number of LODs' spinbox.*/
@@ -559,6 +720,10 @@ private:
 
 	/** LOD Info editing is enabled? LODIndex == -1, then it just verifies if the asset exists */
 	bool IsLODInfoEditingEnabled(int32 LODIndex) const;
+	void ModifyMeshLODSettings(int32 LODIndex);
+
+	TMap<int32, TSharedPtr<FSkeletalMeshBuildSettingsLayout>> BuildSettingsWidgetsPerLOD;
+	TMap<int32, TSharedPtr<FSkeletalMeshReductionSettingsLayout>> ReductionSettingsWidgetsPerLOD;
 
 	// Property handle used to determine if the VertexColorImportOverride property should be enabled.
 	TSharedPtr<IPropertyHandle> VertexColorImportOptionHandle;
diff --git a/Engine/Source/Editor/Persona/Private/SkinWeightProfileHelpers.cpp b/Engine/Source/Editor/Persona/Private/SkinWeightProfileHelpers.cpp
--- a/Engine/Source/Editor/Persona/Private/SkinWeightProfileHelpers.cpp
+++ b/Engine/Source/Editor/Persona/Private/SkinWeightProfileHelpers.cpp
@@ -66,10 +66,10 @@ void FSkinWeightProfileHelpers::ImportSkinWeightProfile(USkeletalMesh* InSkeleta
 		
 		if (bShouldImport)
 		{
+			FScopedSkeletalMeshPostEditChange ScopedPostEditChange(InSkeletalMesh);
 			const FName ProfileName(*ImportSettings->ProfileName);			
 			// Try and import the skin weight profile from the provided FBX file path
-			const bool bReregisterComponent = true;
-			const bool bResult = FSkinWeightsUtilities::ImportAlternateSkinWeight(InSkeletalMesh, PickedFileName, ImportSettings->LODIndex, ProfileName, bReregisterComponent);
+			const bool bResult = FSkinWeightsUtilities::ImportAlternateSkinWeight(InSkeletalMesh, PickedFileName, ImportSettings->LODIndex, ProfileName);
 			if (bResult)
 			{
 				FNotificationInfo NotificationInfo(FText::GetEmpty());
@@ -98,7 +98,7 @@ void FSkinWeightProfileHelpers::ImportSkinWeightProfileLOD(USkeletalMesh* InSkel
 		const FString PickedFileName = FSkinWeightsUtilities::PickSkinWeightFBXPath(LODIndex);
 		
 		// Try and import skin weights for a specific mesh LOD
-		const bool bResult = FSkinWeightsUtilities::ImportAlternateSkinWeight(InSkeletalMesh, PickedFileName, LODIndex, ProfileName, true);
+		const bool bResult = FSkinWeightsUtilities::ImportAlternateSkinWeight(InSkeletalMesh, PickedFileName, LODIndex, ProfileName);
 		if (bResult)
 		{
 			FLODUtilities::RegenerateDependentLODs(InSkeletalMesh, LODIndex);
@@ -126,13 +126,14 @@ void FSkinWeightProfileHelpers::ReimportSkinWeightProfileLOD(USkeletalMesh* InSk
 	{
 		if (const FString* PathNamePtr = Profile->PerLODSourceFiles.Find(LODIndex))
 		{
+			FScopedSkeletalMeshPostEditChange ScopedPostEditChange(InSkeletalMesh);
 			bool bResult = false;
 			const FString& PathName = UAssetImportData::ResolveImportFilename(*PathNamePtr, InSkeletalMesh->GetOutermost());
 			
 			// Check to see if the source file is still valid
 			if (FPaths::FileExists(PathName))
 			{
-				bResult = FSkinWeightsUtilities::ImportAlternateSkinWeight(InSkeletalMesh, PathName, LODIndex, InProfileName, true);
+				bResult = FSkinWeightsUtilities::ImportAlternateSkinWeight(InSkeletalMesh, PathName, LODIndex, InProfileName);
 			}
 			else
 			{
@@ -141,7 +142,7 @@ void FSkinWeightProfileHelpers::ReimportSkinWeightProfileLOD(USkeletalMesh* InSk
 				{
 					// Otherwise let the user pick a new path
 					const FString PickedFileName = FSkinWeightsUtilities::PickSkinWeightFBXPath(LODIndex);
-					bResult = FSkinWeightsUtilities::ImportAlternateSkinWeight(InSkeletalMesh, PickedFileName, LODIndex, InProfileName, true);
+					bResult = FSkinWeightsUtilities::ImportAlternateSkinWeight(InSkeletalMesh, PickedFileName, LODIndex, InProfileName);
 				}
 			}
 
@@ -176,6 +177,7 @@ void FSkinWeightProfileHelpers::RemoveSkinWeightProfile(USkeletalMesh* InSkeleta
 	// Check whether or not we actually removed a profile
 	if (NumRemoved)
 	{
+		FScopedSkeletalMeshPostEditChange ScopedPostEditChange(InSkeletalMesh);
 		FSkeletalMeshUpdateContext UpdateContext;
 		UpdateContext.SkeletalMesh = InSkeletalMesh;
 		
@@ -189,7 +191,7 @@ void FSkinWeightProfileHelpers::RemoveSkinWeightProfile(USkeletalMesh* InSkeleta
 				// Delete the alternate influences data
 				LODModel.SkinWeightProfiles.Remove(InProfileName);
 				// Regenerate this generated LOD
-				FLODUtilities::SimplifySkeletalMeshLOD(UpdateContext, LODIndex, true);
+				FLODUtilities::SimplifySkeletalMeshLOD(UpdateContext, LODIndex);
 				
 				// Goto next LOD
 				continue;
@@ -225,6 +227,8 @@ void FSkinWeightProfileHelpers::RemoveSkinWeightProfileLOD(USkeletalMesh* InSkel
 			return;
 		}
 		
+		FScopedSkeletalMeshPostEditChange ScopedPostEditChange(InSkeletalMesh);
+
 		if (FSkinWeightsUtilities::RemoveSkinnedWeightProfileData(InSkeletalMesh, InProfileName, LODIndex))
 		{
 			FSkinWeightProfileHelpers::ClearSkinWeightProfileInstanceOverrides(InSkeletalMesh, InProfileName);
@@ -249,6 +253,8 @@ void FSkinWeightProfileHelpers::RemoveSkinWeightProfileLOD(USkeletalMesh* InSkel
 
 void FSkinWeightProfileHelpers::ClearSkinWeightProfileInstanceOverrides(USkeletalMesh* InSkeletalMesh, FName InProfileName)
 {
+	//Make sure all component are unregister (world and render data) avoid post edit change
+	FScopedSkeletalMeshPostEditChange ScopedPostEditChange(InSkeletalMesh, false, true);
 	for (TObjectIterator<USkinnedMeshComponent> It; It; ++It)
 	{
 		if (It->SkeletalMesh == InSkeletalMesh)
diff --git a/Engine/Source/Editor/PropertyEditor/Private/PropertyCustomizationHelpers.cpp b/Engine/Source/Editor/PropertyEditor/Private/PropertyCustomizationHelpers.cpp
--- a/Engine/Source/Editor/PropertyEditor/Private/PropertyCustomizationHelpers.cpp
+++ b/Engine/Source/Editor/PropertyEditor/Private/PropertyCustomizationHelpers.cpp
@@ -815,9 +815,9 @@ public:
 	* @param Section		The Section being used
 	* @param bCanBeReplced	Whether or not the Section can be replaced by a user
 	*/
-	virtual void AddSection(int32 LodIndex, int32 SectionIndex, FName InMaterialSlotName, int32 InMaterialSlotIndex, FName InOriginalMaterialSlotName, const TMap<int32, FName> &InAvailableMaterialSlotName, const UMaterialInterface* Material, bool IsSectionUsingCloth) override
+	virtual void AddSection(int32 LodIndex, int32 SectionIndex, FName InMaterialSlotName, int32 InMaterialSlotIndex, FName InOriginalMaterialSlotName, const TMap<int32, FName> &InAvailableMaterialSlotName, const UMaterialInterface* Material, bool IsSectionUsingCloth, bool bIsChunkSection, int32 DefaultMaterialIndex) override
 	{
-		FSectionListItem SectionItem(LodIndex, SectionIndex, InMaterialSlotName, InMaterialSlotIndex, InOriginalMaterialSlotName, InAvailableMaterialSlotName, Material, IsSectionUsingCloth, ThumbnailSize);
+		FSectionListItem SectionItem(LodIndex, SectionIndex, InMaterialSlotName, InMaterialSlotIndex, InOriginalMaterialSlotName, InAvailableMaterialSlotName, Material, IsSectionUsingCloth, ThumbnailSize, bIsChunkSection, DefaultMaterialIndex);
 		if (!Sections.Contains(SectionItem))
 		{
 			Sections.Add(SectionItem);
@@ -934,7 +934,10 @@ public:
 
 	TSharedRef<SWidget> CreateValueContent(const TSharedPtr<FAssetThumbnailPool>& ThumbnailPool)
 	{
-		FText MaterialSlotNameTooltipText = SectionItem.IsSectionUsingCloth ? FText(LOCTEXT("SectionIndex_MaterialSlotNameTooltip", "Cannot change the material slot when the mesh section use the cloth system.")) : FText::GetEmpty();
+		FFormatNamedArguments Arguments;
+		Arguments.Add(TEXT("DefaultMaterialIndex"), SectionItem.DefaultMaterialIndex);
+		FText BaseMaterialSlotTooltip = SectionItem.DefaultMaterialIndex != SectionItem.MaterialSlotIndex ? FText::Format(LOCTEXT("SectionIndex_BaseMaterialSlotNameTooltip", "This section material slot was change from the default value [{DefaultMaterialIndex}]."), Arguments) : FText::GetEmpty();
+		FText MaterialSlotNameTooltipText = SectionItem.IsSectionUsingCloth ? FText(LOCTEXT("SectionIndex_MaterialSlotNameTooltip", "Cannot change the material slot when the mesh section use the cloth system.")) : BaseMaterialSlotTooltip;
 		return
 			SNew(SHorizontalBox)
 			+ SHorizontalBox::Slot()
@@ -948,6 +951,7 @@ public:
 				.HAlign(HAlign_Fill)
 				[
 					SNew(SHorizontalBox)
+					.Visibility(SectionItem.bIsChunkSection ? EVisibility::Collapsed : EVisibility::Visible)
 					+ SHorizontalBox::Slot()
 					.FillWidth(1.0f)
 					[
@@ -1023,6 +1027,22 @@ public:
 						]
 					]
 				]
+				+ SVerticalBox::Slot()
+				.AutoHeight()
+				.Padding(0.0f)
+				.VAlign(VAlign_Center)
+				.HAlign(HAlign_Fill)
+				[
+					SNew(SHorizontalBox)
+					.Visibility(SectionItem.bIsChunkSection ? EVisibility::Visible : EVisibility::Collapsed)
+					+ SHorizontalBox::Slot()
+					.FillWidth(1.0f)
+					[
+						SNew(STextBlock)
+						.Font(IDetailLayoutBuilder::GetDetailFont())
+						.Text(LOCTEXT("SectionListItemChunkSectionValueLabel", "Chunked"))
+					]
+				]
 			];
 	}
 
@@ -1076,7 +1096,12 @@ private:
 	{
 		FString MaterialSlotDisplayName;
 		SectionItem.MaterialSlotName.ToString(MaterialSlotDisplayName);
-		MaterialSlotDisplayName = TEXT("[") + FString::FromInt(SectionItem.MaterialSlotIndex) + TEXT("] ") + MaterialSlotDisplayName;
+		FString MaterialSlotRemapString = TEXT("");
+		if (SectionItem.DefaultMaterialIndex != SectionItem.MaterialSlotIndex)
+		{
+			MaterialSlotRemapString = TEXT(" (Modified)");
+		}
+		MaterialSlotDisplayName = TEXT("[") + FString::FromInt(SectionItem.MaterialSlotIndex) + TEXT("] ") + MaterialSlotDisplayName + MaterialSlotRemapString;
 		return FText::FromString(MaterialSlotDisplayName);
 	}
 
@@ -1210,7 +1235,7 @@ void FSectionList::GenerateChildContent(IDetailChildrenBuilder& ChildrenBuilder)
 			if (bDisplayAllSectionsInSlot)
 			{
 				FDetailWidgetRow& ChildRow = ChildrenBuilder.AddCustomRow(Section.Material.IsValid() ? FText::FromString(Section.Material->GetName()) : FText::GetEmpty());
-				AddSectionItem(ChildRow, CurrentLODIndex, FSectionListItem(CurrentLODIndex, Section.SectionIndex, Section.MaterialSlotName, Section.MaterialSlotIndex, Section.OriginalMaterialSlotName, Section.AvailableMaterialSlotName, Section.Material.Get(), Section.IsSectionUsingCloth, ThumbnailSize), !bDisplayAllSectionsInSlot);
+				AddSectionItem(ChildRow, CurrentLODIndex, FSectionListItem(CurrentLODIndex, Section.SectionIndex, Section.MaterialSlotName, Section.MaterialSlotIndex, Section.OriginalMaterialSlotName, Section.AvailableMaterialSlotName, Section.Material.Get(), Section.IsSectionUsingCloth, ThumbnailSize, Section.bIsChunkSection, Section.DefaultMaterialIndex), !bDisplayAllSectionsInSlot);
 			}
 		}
 	}
@@ -1292,6 +1317,7 @@ void FSectionList::AddSectionItem(FDetailWidgetRow& Row, int32 LodIndex, const s
 {
 	uint32 NumSections = SectionListBuilder->GetNumSections(LodIndex);
 
+	bool bIsChunkSection = Item.bIsChunkSection;
 	TSharedRef<FSectionItemView> NewView = FSectionItemView::Create(Item, SectionListDelegates.OnSectionChanged, SectionListDelegates.OnGenerateCustomNameWidgets, SectionListDelegates.OnGenerateCustomSectionWidgets, SectionListDelegates.OnResetSectionToDefaultClicked, NumSections, ThumbnailSize);
 
 	TSharedPtr<SWidget> RightSideContent;
@@ -1318,13 +1344,17 @@ void FSectionList::AddSectionItem(FDetailWidgetRow& Row, int32 LodIndex, const s
 		ViewedSections.Add(NewView);
 	}
 
-	Row.CopyAction(FUIAction(FExecuteAction::CreateSP(this, &FSectionList::OnCopySectionItem, LodIndex, Item.SectionIndex), FCanExecuteAction::CreateSP(this, &FSectionList::OnCanCopySectionItem, LodIndex, Item.SectionIndex)));
-	Row.PasteAction(FUIAction(FExecuteAction::CreateSP(this, &FSectionList::OnPasteSectionItem, LodIndex, Item.SectionIndex)));
-
-	if(SectionListDelegates.OnEnableSectionItem.IsBound())
+	//Chunk section cannot be copy enable or disable, do the operation on the parent section
+	if (!bIsChunkSection)
 	{
-		Row.AddCustomContextMenuAction(FUIAction(FExecuteAction::CreateSP(this, &FSectionList::OnEnableSectionItem, LodIndex, Item.SectionIndex, true)), LOCTEXT("SectionItemContexMenu_Enable", "Enable"));
-		Row.AddCustomContextMenuAction(FUIAction(FExecuteAction::CreateSP(this, &FSectionList::OnEnableSectionItem, LodIndex, Item.SectionIndex, false)), LOCTEXT("SectionItemContexMenu_Disable", "Disable"));
+		Row.CopyAction(FUIAction(FExecuteAction::CreateSP(this, &FSectionList::OnCopySectionItem, LodIndex, Item.SectionIndex), FCanExecuteAction::CreateSP(this, &FSectionList::OnCanCopySectionItem, LodIndex, Item.SectionIndex)));
+		Row.PasteAction(FUIAction(FExecuteAction::CreateSP(this, &FSectionList::OnPasteSectionItem, LodIndex, Item.SectionIndex)));
+
+		if (SectionListDelegates.OnEnableSectionItem.IsBound())
+		{
+			Row.AddCustomContextMenuAction(FUIAction(FExecuteAction::CreateSP(this, &FSectionList::OnEnableSectionItem, LodIndex, Item.SectionIndex, true)), LOCTEXT("SectionItemContexMenu_Enable", "Enable"));
+			Row.AddCustomContextMenuAction(FUIAction(FExecuteAction::CreateSP(this, &FSectionList::OnEnableSectionItem, LodIndex, Item.SectionIndex, false)), LOCTEXT("SectionItemContexMenu_Disable", "Disable"));
+		}
 	}
 
 	Row.NameContent()
diff --git a/Engine/Source/Editor/PropertyEditor/Public/PropertyCustomizationHelpers.h b/Engine/Source/Editor/PropertyEditor/Public/PropertyCustomizationHelpers.h
--- a/Engine/Source/Editor/PropertyEditor/Public/PropertyCustomizationHelpers.h
+++ b/Engine/Source/Editor/PropertyEditor/Public/PropertyCustomizationHelpers.h
@@ -576,7 +576,7 @@ public:
 	* @param Section		The Section being used
 	* @param bCanBeReplced	Whether or not the Section can be replaced by a user
 	*/
-	virtual void AddSection(int32 LodIndex, int32 SectionIndex, FName InMaterialSlotName, int32 InMaterialSlotIndex, FName InOriginalMaterialSlotName, const TMap<int32, FName> &InAvailableMaterialSlotName, const UMaterialInterface* Material, bool IsSectionUsingCloth) = 0;
+	virtual void AddSection(int32 LodIndex, int32 SectionIndex, FName InMaterialSlotName, int32 InMaterialSlotIndex, FName InOriginalMaterialSlotName, const TMap<int32, FName> &InAvailableMaterialSlotName, const UMaterialInterface* Material, bool IsSectionUsingCloth, bool bIsChunkSection, int32 DefaultMaterialIndex) = 0;
 };
 
 
@@ -607,8 +607,11 @@ struct FSectionListItem
 	/* Available material slot name*/
 	TMap<int32, FName> AvailableMaterialSlotName;
 
+	bool bIsChunkSection;
+	int32 DefaultMaterialIndex;
 
-	FSectionListItem(int32 InLodIndex, int32 InSectionIndex, FName InMaterialSlotName, int32 InMaterialSlotIndex, FName InOriginalMaterialSlotName, const TMap<int32, FName> &InAvailableMaterialSlotName, const UMaterialInterface* InMaterial, bool InIsSectionUsingCloth, int32 InThumbnailSize)
+
+	FSectionListItem(int32 InLodIndex, int32 InSectionIndex, FName InMaterialSlotName, int32 InMaterialSlotIndex, FName InOriginalMaterialSlotName, const TMap<int32, FName> &InAvailableMaterialSlotName, const UMaterialInterface* InMaterial, bool InIsSectionUsingCloth, int32 InThumbnailSize, bool InIsChunkSection, int32 InDefaultMaterialIndex)
 		: LodIndex(InLodIndex)
 		, SectionIndex(InSectionIndex)
 		, IsSectionUsingCloth(InIsSectionUsingCloth)
@@ -618,11 +621,21 @@ struct FSectionListItem
 		, MaterialSlotIndex(InMaterialSlotIndex)
 		, OriginalMaterialSlotName(InOriginalMaterialSlotName)
 		, AvailableMaterialSlotName(InAvailableMaterialSlotName)
+		, bIsChunkSection(InIsChunkSection)
+		, DefaultMaterialIndex(InDefaultMaterialIndex)
 	{}
 
 	bool operator==(const FSectionListItem& Other) const
 	{
-		bool IsSectionItemEqual = LodIndex == Other.LodIndex && SectionIndex == Other.SectionIndex && MaterialSlotIndex == Other.MaterialSlotIndex && MaterialSlotName == Other.MaterialSlotName && Material == Other.Material && AvailableMaterialSlotName.Num() == Other.AvailableMaterialSlotName.Num() && IsSectionUsingCloth == Other.IsSectionUsingCloth;
+		bool IsSectionItemEqual = LodIndex == Other.LodIndex
+			&& SectionIndex == Other.SectionIndex
+			&& MaterialSlotIndex == Other.MaterialSlotIndex
+			&& MaterialSlotName == Other.MaterialSlotName
+			&& Material == Other.Material
+			&& AvailableMaterialSlotName.Num() == Other.AvailableMaterialSlotName.Num()
+			&& IsSectionUsingCloth == Other.IsSectionUsingCloth
+			&& bIsChunkSection == Other.bIsChunkSection
+			&& DefaultMaterialIndex == Other.DefaultMaterialIndex;
 		if (IsSectionItemEqual)
 		{
 			for (auto Kvp : AvailableMaterialSlotName)
diff --git a/Engine/Source/Editor/SkeletalMeshEditor/Private/SkeletalMeshEditor.cpp b/Engine/Source/Editor/SkeletalMeshEditor/Private/SkeletalMeshEditor.cpp
--- a/Engine/Source/Editor/SkeletalMeshEditor/Private/SkeletalMeshEditor.cpp
+++ b/Engine/Source/Editor/SkeletalMeshEditor/Private/SkeletalMeshEditor.cpp
@@ -44,6 +44,8 @@
 #include "EditorFramework/AssetImportData.h"
 #include "Factories/FbxSkeletalMeshImportData.h"
 
+#include "Misc/MessageDialog.h"
+
 const FName SkeletalMeshEditorAppIdentifier = FName(TEXT("SkeletalMeshEditorApp"));
 
 const FName SkeletalMeshEditorModes::SkeletalMeshEditorMode(TEXT("SkeletalMeshEditorMode"));
@@ -78,6 +80,112 @@ FSkeletalMeshEditor::~FSkeletalMeshEditor()
 	}
 }
 
+bool IsReductionParentBaseLODUseSkeletalMeshBuildWorkflow(USkeletalMesh* SkeletalMesh, int32 TestLODIndex)
+{
+	FSkeletalMeshLODInfo* LODInfo = SkeletalMesh->GetLODInfo(TestLODIndex);
+	if (LODInfo == nullptr || !SkeletalMesh->GetImportedModel() || !SkeletalMesh->GetImportedModel()->LODModels.IsValidIndex(TestLODIndex))
+	{
+		return false;
+	}
+
+	if (SkeletalMesh->GetImportedModel()->LODModels[TestLODIndex].RawSkeletalMeshBulkData.IsBuildDataAvailable())
+	{
+		return true;
+	}
+	if (LODInfo->bHasBeenSimplified || SkeletalMesh->IsReductionActive(TestLODIndex))
+	{
+		int32 ReduceBaseLOD = LODInfo->ReductionSettings.BaseLOD;
+		if (ReduceBaseLOD < TestLODIndex)
+		{
+			return IsReductionParentBaseLODUseSkeletalMeshBuildWorkflow(SkeletalMesh, ReduceBaseLOD);
+		}
+	}
+	return false;
+}
+
+bool FSkeletalMeshEditor::OnRequestClose()
+{
+	bool bAllowClose = true;
+
+	if (PersonaToolkit.IsValid() && SkeletalMesh)
+	{
+		bool bHaveModifiedLOD = false;
+		for (int32 LODIndex = 0; LODIndex < SkeletalMesh->GetLODNum(); ++LODIndex)
+		{
+			FSkeletalMeshLODInfo* LODInfo = SkeletalMesh->GetLODInfo(LODIndex);
+			if (LODInfo == nullptr || !SkeletalMesh->GetImportedModel() || !SkeletalMesh->GetImportedModel()->LODModels.IsValidIndex(LODIndex))
+			{
+				continue;
+			}
+			
+			//Do not prevent exiting if we are not using the skeletal mesh build workflow
+			if (!SkeletalMesh->GetImportedModel()->LODModels[LODIndex].RawSkeletalMeshBulkData.IsBuildDataAvailable())
+			{
+				if (!LODInfo->bHasBeenSimplified && !SkeletalMesh->IsReductionActive(LODIndex))
+				{
+					continue;
+				}
+				//Do not prevent exit if the generated LOD is base on a LODModel not using the skeletalmesh build workflow
+				int32 ReduceBaseLOD = LODInfo->ReductionSettings.BaseLOD;
+				if(!IsReductionParentBaseLODUseSkeletalMeshBuildWorkflow(SkeletalMesh, ReduceBaseLOD))
+				{
+					continue;
+				}
+				
+			}
+
+			bool bValidLODSettings = false;
+			if (SkeletalMesh->LODSettings != nullptr)
+			{
+				const int32 NumSettings = FMath::Min(SkeletalMesh->LODSettings->GetNumberOfSettings(), SkeletalMesh->GetLODNum());
+				if (LODIndex < NumSettings)
+				{
+					bValidLODSettings = true;
+				}
+			}
+			const FSkeletalMeshLODGroupSettings* SkeletalMeshLODGroupSettings = bValidLODSettings ? &SkeletalMesh->LODSettings->GetSettingsForLODLevel(LODIndex) : nullptr;
+
+			FGuid BuildGUID = LODInfo->ComputeDeriveDataCacheKey(SkeletalMeshLODGroupSettings);
+			if (LODInfo->BuildGUID != BuildGUID)
+			{
+				bHaveModifiedLOD = true;
+				break;
+			}
+			FString BuildStringID = SkeletalMesh->GetImportedModel()->LODModels[LODIndex].GetLODModelDeriveDataKey();
+			if (SkeletalMesh->GetImportedModel()->LODModels[LODIndex].BuildStringID != BuildStringID)
+			{
+				bHaveModifiedLOD = true;
+				break;
+			}
+		}
+
+		if (bHaveModifiedLOD)
+		{
+			// find out the user wants to do with this dirty material
+			EAppReturnType::Type OkCancelReply = FMessageDialog::Open(
+				EAppMsgType::OkCancel,
+				FText::Format(LOCTEXT("SkeletalMeshEditorShouldApplyLODChanges", "We have to apply level of detail changes to {0} before exiting the skeletal mesh editor."), FText::FromString(PersonaToolkit->GetMesh()->GetName()))
+			);
+
+			switch (OkCancelReply)
+			{
+			case EAppReturnType::Ok:
+				{
+					SkeletalMesh->MarkPackageDirty();
+					SkeletalMesh->PostEditChange();
+					bAllowClose = true;
+				}
+				break;
+			case EAppReturnType::Cancel:
+				// Don't exit.
+				bAllowClose = false;
+				break;
+			}
+		}
+	}
+	return bAllowClose;
+}
+
 void FSkeletalMeshEditor::RegisterTabSpawners(const TSharedRef<class FTabManager>& InTabManager)
 {
 	WorkspaceMenuCategory = InTabManager->AddLocalWorkspaceMenuCategory(LOCTEXT("WorkspaceMenu_SkeletalMeshEditor", "Skeletal Mesh Editor"));
@@ -546,8 +654,7 @@ void FSkeletalMeshEditor::OnRemoveSectionFromLodAndBelowMenuItemClicked(int32 Lo
 		const FSkeletalMeshLODInfo* CurrentSkeletalMeshLODInfo = SkeletalMesh->GetLODInfo(GenerateLodIndex);
 		if (CurrentSkeletalMeshLODInfo != nullptr && CurrentSkeletalMeshLODInfo->bHasBeenSimplified && BaseLodIndexes.Contains(CurrentSkeletalMeshLODInfo->ReductionSettings.BaseLOD))
 		{
-			bool bRestoreClothing = true;
-			FLODUtilities::SimplifySkeletalMeshLOD(UpdateContext, GenerateLodIndex, true, bRestoreClothing);
+			FLODUtilities::SimplifySkeletalMeshLOD(UpdateContext, GenerateLodIndex, true);
 			BaseLodIndexes.Add(GenerateLodIndex);
 		}
 	}
@@ -668,7 +775,7 @@ void FSkeletalMeshEditor::OnCreateClothingAssetMenuItemClicked(FSkeletalMeshClot
 						{
 							SectionIndex = i;
 							AssetLodIndex = SkelLod.Sections[i].ClothingData.AssetLodIndex;
-							TargetAssetPtr->UnbindFromSkeletalMesh(Mesh, Params.TargetLod);
+							RemoveClothing(Params.TargetLod, SectionIndex);
 							break;
 						}
 					}
@@ -786,13 +893,54 @@ void FSkeletalMeshEditor::ApplyClothing(UClothingAssetBase* InAsset, int32 InLod
 {
 	USkeletalMesh* Mesh = GetPersonaToolkit()->GetPreviewMesh();
 
-	if(UClothingAsset* ClothingAsset = Cast<UClothingAsset>(InAsset))
+	if (Mesh == nullptr || Mesh->GetImportedModel() == nullptr || !Mesh->GetImportedModel()->LODModels.IsValidIndex(InLodIndex))
 	{
-		ClothingAsset->BindToSkeletalMesh(Mesh, InLodIndex, InSectionIndex, InClothingLod);
+		return;
 	}
-	else
+
+	FSkeletalMeshLODModel& LODModel = Mesh->GetImportedModel()->LODModels[InLodIndex];
+	const FSkelMeshSection& Section = LODModel.Sections[InSectionIndex];
+	FScopedSuspendAlternateSkinWeightPreview ScopedSuspendAlternateSkinnWeightPreview(Mesh);
 	{
-		RemoveClothing(InLodIndex, InSectionIndex);
+		FScopedSkeletalMeshPostEditChange ScopedPostEditChange(Mesh);
+		FScopedTransaction Transaction(LOCTEXT("SkeletalMeshEditorApplyClothingTransaction", "Persona editor: Apply Section Cloth"));
+		Mesh->Modify();
+
+		FSkelMeshSourceSectionUserData& OriginalSectionData = LODModel.UserSectionsData.FindOrAdd(Section.OriginalDataSectionIndex);
+		auto ClearOriginalSectionUserData = [&OriginalSectionData]()
+		{
+			OriginalSectionData.CorrespondClothAssetIndex = INDEX_NONE;
+			OriginalSectionData.ClothingData.AssetGuid = FGuid();
+			OriginalSectionData.ClothingData.AssetLodIndex = INDEX_NONE;
+		};
+		if (UClothingAsset* ClothingAsset = Cast<UClothingAsset>(InAsset))
+		{
+			// Look for a currently bound asset an unbind it if necessary first
+			if (UClothingAssetBase* CurrentAsset = Mesh->GetSectionClothingAsset(InLodIndex, InSectionIndex))
+			{
+				CurrentAsset->UnbindFromSkeletalMesh(Mesh, InLodIndex);
+				ClearOriginalSectionUserData();
+			}
+
+			if (ClothingAsset->BindToSkeletalMesh(Mesh, InLodIndex, InSectionIndex, InClothingLod))
+			{
+				//Successful bind so set the SectionUserData
+				int32 AssetIndex = INDEX_NONE;
+				check(Mesh->MeshClothingAssets.Find(ClothingAsset, AssetIndex));
+				OriginalSectionData.CorrespondClothAssetIndex = AssetIndex;
+				OriginalSectionData.ClothingData.AssetGuid = ClothingAsset->GetAssetGuid();
+				OriginalSectionData.ClothingData.AssetLodIndex = InClothingLod;
+			}
+		}
+		else if (Mesh)
+		{
+			//User set none, so unbind anything that is bind
+			if (UClothingAssetBase* CurrentAsset = Mesh->GetSectionClothingAsset(InLodIndex, InSectionIndex))
+			{
+				CurrentAsset->UnbindFromSkeletalMesh(Mesh, InLodIndex);
+				ClearOriginalSectionUserData();
+			}
+		}
 	}
 }
 
@@ -800,11 +948,31 @@ void FSkeletalMeshEditor::RemoveClothing(int32 InLodIndex, int32 InSectionIndex)
 {
 	USkeletalMesh* Mesh = GetPersonaToolkit()->GetPreviewMesh();
 
-	if(Mesh)
+	if (Mesh->GetImportedModel() == nullptr || !Mesh->GetImportedModel()->LODModels.IsValidIndex(InLodIndex))
 	{
-		if(UClothingAssetBase* CurrentAsset = Mesh->GetSectionClothingAsset(InLodIndex, InSectionIndex))
+		return;
+	}
+
+	FSkeletalMeshLODModel& LODModel = Mesh->GetImportedModel()->LODModels[InLodIndex];
+	const FSkelMeshSection& Section = LODModel.Sections[InSectionIndex];
+	FScopedSuspendAlternateSkinWeightPreview ScopedSuspendAlternateSkinnWeightPreview(Mesh);
+	{
+		FScopedSkeletalMeshPostEditChange ScopedPostEditChange(Mesh);
+		FScopedTransaction Transaction(LOCTEXT("SkeletalMeshEditorRemoveClothingTransaction", "Persona editor: Remove Section Cloth"));
+		Mesh->Modify();
+
+		FSkelMeshSourceSectionUserData& OriginalSectionData = LODModel.UserSectionsData.FindOrAdd(Section.OriginalDataSectionIndex);
+		auto ClearOriginalSectionUserData = [&OriginalSectionData]()
+		{
+			OriginalSectionData.CorrespondClothAssetIndex = INDEX_NONE;
+			OriginalSectionData.ClothingData.AssetGuid = FGuid();
+			OriginalSectionData.ClothingData.AssetLodIndex = INDEX_NONE;
+		};
+		// Look for a currently bound asset an unbind it if necessary first
+		if (UClothingAssetBase* CurrentAsset = Mesh->GetSectionClothingAsset(InLodIndex, InSectionIndex))
 		{
 			CurrentAsset->UnbindFromSkeletalMesh(Mesh, InLodIndex);
+			ClearOriginalSectionUserData();
 		}
 	}
 }
@@ -883,92 +1051,92 @@ UObject* FSkeletalMeshEditor::HandleGetAsset()
 
 bool FSkeletalMeshEditor::HandleReimportMeshInternal(int32 SourceFileIndex /*= INDEX_NONE*/, bool bWithNewFile /*= false*/)
 {
-	// Make sure we clear any skin weight profile currently being previewed
-	UDebugSkelMeshComponent* MeshComponent = GetPersonaToolkit()->GetPreviewScene()->GetPreviewMeshComponent();
-	const FName ProfileName = MeshComponent ? MeshComponent->GetCurrentSkinWeightProfileName() : NAME_None;
-	if (MeshComponent)
-	{
-		MeshComponent->ClearSkinWeightProfile();
-	}
-
-	// Reimport the asset
-	const bool bResult = FReimportManager::Instance()->Reimport(SkeletalMesh, true, true, TEXT(""), nullptr, SourceFileIndex, bWithNewFile);
-
-	// Refresh skeleton tree
-	SkeletonTree->Refresh();
-
-	// Re-set skin weigh profile
-	if (MeshComponent && (ProfileName != NAME_None))
+	FScopedSuspendAlternateSkinWeightPreview ScopedSuspendAlternateSkinnWeightPreview(SkeletalMesh);
+	bool bResult = false;
 	{
-		MeshComponent->SetSkinWeightProfile(ProfileName);
+		FScopedSkeletalMeshPostEditChange ScopedPostEditChange(SkeletalMesh);
+		// Reimport the asset
+		bResult = FReimportManager::Instance()->Reimport(SkeletalMesh, true, true, TEXT(""), nullptr, SourceFileIndex, bWithNewFile);
+		// Refresh skeleton tree
+		SkeletonTree->Refresh();
 	}
-
 	return bResult;
 }
 
 void FSkeletalMeshEditor::HandleReimportMesh(int32 SourceFileIndex /*= INDEX_NONE*/)
 {
+	FScopedSuspendAlternateSkinWeightPreview ScopedSuspendAlternateSkinnWeightPreview(SkeletalMesh);
+
 	HandleReimportMeshInternal(SourceFileIndex, false);
 }
 
 void FSkeletalMeshEditor::HandleReimportMeshWithNewFile(int32 SourceFileIndex /*= INDEX_NONE*/)
 {
+	FScopedSuspendAlternateSkinWeightPreview ScopedSuspendAlternateSkinnWeightPreview(SkeletalMesh);
+
 	HandleReimportMeshInternal(SourceFileIndex, true);
 }
 
 void ReimportAllCustomLODs(USkeletalMesh* SkeletalMesh, UDebugSkelMeshComponent* PreviewMeshComponent, bool bWithNewFile)
 {
-	//Find the dependencies of the generated LOD
-	TArray<bool> Dependencies;
-	Dependencies.AddZeroed(SkeletalMesh->GetLODNum());
-	//Avoid making LOD 0 to true in the dependencies since everything that should be regenerate base on LOD 0 is already regenerate at this point.
-	//But we need to regenerate every generated LOD base on any re-import custom LOD
-	//Reimport all custom LODs
-	for (int32 LodIndex = 1; LodIndex < SkeletalMesh->GetLODNum(); ++LodIndex)
+	FScopedSuspendAlternateSkinWeightPreview ScopedSuspendAlternateSkinnWeightPreview(SkeletalMesh);
 	{
-		//Do not reimport LOD that was re-import with the base mesh
-		if (SkeletalMesh->GetLODInfo(LodIndex)->bImportWithBaseMesh)
-		{
-			continue;
-		}
-		if (SkeletalMesh->GetLODInfo(LodIndex)->bHasBeenSimplified == false)
+		FScopedSkeletalMeshPostEditChange ScopedPostEditChange(SkeletalMesh);
+
+		//Find the dependencies of the generated LOD
+		TArray<bool> Dependencies;
+		Dependencies.AddZeroed(SkeletalMesh->GetLODNum());
+		//Avoid making LOD 0 to true in the dependencies since everything that should be regenerate base on LOD 0 is already regenerate at this point.
+		//But we need to regenerate every generated LOD base on any re-import custom LOD
+		//Reimport all custom LODs
+		for (int32 LodIndex = 1; LodIndex < SkeletalMesh->GetLODNum(); ++LodIndex)
 		{
-			FString SourceFilenameBackup = SkeletalMesh->GetLODInfo(LodIndex)->SourceImportFilename;
-			if (bWithNewFile)
+			//Do not reimport LOD that was re-import with the base mesh
+			if (SkeletalMesh->GetLODInfo(LodIndex)->bImportWithBaseMesh)
 			{
-				SkeletalMesh->GetLODInfo(LodIndex)->SourceImportFilename.Empty();
+				continue;
 			}
-
-			if (!FbxMeshUtils::ImportMeshLODDialog(SkeletalMesh, LodIndex))
+			if (SkeletalMesh->GetLODInfo(LodIndex)->bHasBeenSimplified == false)
 			{
+				FString SourceFilenameBackup = SkeletalMesh->GetLODInfo(LodIndex)->SourceImportFilename;
 				if (bWithNewFile)
 				{
-					SkeletalMesh->GetLODInfo(LodIndex)->SourceImportFilename = SourceFilenameBackup;
+					SkeletalMesh->GetLODInfo(LodIndex)->SourceImportFilename.Empty();
+				}
+
+				if (!FbxMeshUtils::ImportMeshLODDialog(SkeletalMesh, LodIndex, false))
+				{
+					if (bWithNewFile)
+					{
+						SkeletalMesh->GetLODInfo(LodIndex)->SourceImportFilename = SourceFilenameBackup;
+					}
+				}
+				else
+				{
+					Dependencies[LodIndex] = true;
 				}
 			}
-			else
+			else if (Dependencies[SkeletalMesh->GetLODInfo(LodIndex)->ReductionSettings.BaseLOD])
 			{
+				//Regenerate the LOD
+				FSkeletalMeshUpdateContext UpdateContext;
+				UpdateContext.SkeletalMesh = SkeletalMesh;
+				UpdateContext.AssociatedComponents.Push(PreviewMeshComponent);
+				FLODUtilities::SimplifySkeletalMeshLOD(UpdateContext, LodIndex, false);
 				Dependencies[LodIndex] = true;
 			}
 		}
-		else if(Dependencies[SkeletalMesh->GetLODInfo(LodIndex)->ReductionSettings.BaseLOD])
-		{
-			//Regenerate the LOD
-			FSkeletalMeshUpdateContext UpdateContext;
-			UpdateContext.SkeletalMesh = SkeletalMesh;
-			UpdateContext.AssociatedComponents.Push(PreviewMeshComponent);
-			bool bRestoreClothing = true;
-			FLODUtilities::SimplifySkeletalMeshLOD(UpdateContext, LodIndex, true, bRestoreClothing);
-			Dependencies[LodIndex] = true;
-		}
 	}
 }
 
 void FSkeletalMeshEditor::HandleReimportAllMesh(int32 SourceFileIndex /*= INDEX_NONE*/)
 {
+	FScopedSuspendAlternateSkinWeightPreview ScopedSuspendAlternateSkinnWeightPreview(SkeletalMesh);
 	// Reimport the asset
 	if (SkeletalMesh)
 	{
+		FScopedSkeletalMeshPostEditChange ScopedPostEditChange(SkeletalMesh);
+
 		//Reimport base LOD
 		if (HandleReimportMeshInternal(SourceFileIndex, false))
 		{
@@ -980,9 +1148,12 @@ void FSkeletalMeshEditor::HandleReimportAllMesh(int32 SourceFileIndex /*= INDEX_
 
 void FSkeletalMeshEditor::HandleReimportAllMeshWithNewFile(int32 SourceFileIndex /*= INDEX_NONE*/)
 {
+	FScopedSuspendAlternateSkinWeightPreview ScopedSuspendAlternateSkinnWeightPreview(SkeletalMesh);
+
 	// Reimport the asset
 	if (SkeletalMesh)
 	{
+		FScopedSkeletalMeshPostEditChange ScopedPostEditChange(SkeletalMesh);
 		TArray<UObject*> ImportObjs;
 		ImportObjs.Add(SkeletalMesh);
 		if (HandleReimportMeshInternal(SourceFileIndex, true))
diff --git a/Engine/Source/Editor/SkeletalMeshEditor/Private/SkeletalMeshEditor.h b/Engine/Source/Editor/SkeletalMeshEditor/Private/SkeletalMeshEditor.h
--- a/Engine/Source/Editor/SkeletalMeshEditor/Private/SkeletalMeshEditor.h
+++ b/Engine/Source/Editor/SkeletalMeshEditor/Private/SkeletalMeshEditor.h
@@ -62,6 +62,10 @@ public:
 	virtual FString GetWorldCentricTabPrefix() const override;
 	virtual FLinearColor GetWorldCentricTabColorScale() const override;
 
+	//~ Begin FAssetEditorToolkit Interface.
+	virtual bool OnRequestClose() override;
+	//~ End FAssetEditorToolkit Interface.
+
 	/** FEditorUndoClient interface */
 	virtual void PostUndo(bool bSuccess) override;
 	virtual void PostRedo(bool bSuccess) override;
diff --git a/Engine/Source/Editor/StaticMeshEditor/Private/StaticMeshEditor.cpp b/Engine/Source/Editor/StaticMeshEditor/Private/StaticMeshEditor.cpp
--- a/Engine/Source/Editor/StaticMeshEditor/Private/StaticMeshEditor.cpp
+++ b/Engine/Source/Editor/StaticMeshEditor/Private/StaticMeshEditor.cpp
@@ -1398,7 +1398,7 @@ void FStaticMeshEditor::HandleReimportAllMesh()
 					continue;
 				}
 
-				bool bHasBeenSimplified = StaticMesh->GetMeshDescription(LodIndex) == nullptr || StaticMesh->IsReductionActive(LodIndex);
+				bool bHasBeenSimplified = !StaticMesh->IsMeshDescriptionValid(LodIndex) || StaticMesh->IsReductionActive(LodIndex);
 				if (!bHasBeenSimplified)
 				{
 					FbxMeshUtils::ImportMeshLODDialog(StaticMesh, LodIndex);
diff --git a/Engine/Source/Editor/StaticMeshEditor/Private/StaticMeshEditorTools.cpp b/Engine/Source/Editor/StaticMeshEditor/Private/StaticMeshEditorTools.cpp
--- a/Engine/Source/Editor/StaticMeshEditor/Private/StaticMeshEditorTools.cpp
+++ b/Engine/Source/Editor/StaticMeshEditor/Private/StaticMeshEditorTools.cpp
@@ -1971,7 +1971,8 @@ void FMeshSectionSettingsLayout::OnGetSectionsForView(ISectionListBuilder& OutSe
 				{
 					SectionMaterial = UMaterial::GetDefaultMaterial(MD_Surface);
 				}
-				OutSections.AddSection(LODIndex, SectionIndex, CurrentSectionMaterialSlotName, MaterialIndex, CurrentSectionOriginalImportedMaterialName, AvailableSectionName, StaticMesh.StaticMaterials[MaterialIndex].MaterialInterface, false);
+				//TODO: Need to know if a section material slot assignment was change from the default value (implemented in skeletalmesh editor)
+				OutSections.AddSection(LODIndex, SectionIndex, CurrentSectionMaterialSlotName, MaterialIndex, CurrentSectionOriginalImportedMaterialName, AvailableSectionName, StaticMesh.StaticMaterials[MaterialIndex].MaterialInterface, false, false, MaterialIndex);
 			}
 		}
 	}
diff --git a/Engine/Source/Editor/UnrealEd/Private/EditorExporters.cpp b/Engine/Source/Editor/UnrealEd/Private/EditorExporters.cpp
--- a/Engine/Source/Editor/UnrealEd/Private/EditorExporters.cpp
+++ b/Engine/Source/Editor/UnrealEd/Private/EditorExporters.cpp
@@ -59,7 +59,7 @@
 #include "Editor.h"
 #include "MatineeExporter.h"
 #include "FbxExporter.h"
-#include "MeshDescription.h"
+#include "StaticMeshAttributes.h"
 #include "MeshDescriptionOperations.h"
 #include "MaterialUtilities.h"
 #include "InstancedFoliageActor.h"
diff --git a/Engine/Source/Editor/UnrealEd/Private/Factories/EditorFactories.cpp b/Engine/Source/Editor/UnrealEd/Private/Factories/EditorFactories.cpp
--- a/Engine/Source/Editor/UnrealEd/Private/Factories/EditorFactories.cpp
+++ b/Engine/Source/Editor/UnrealEd/Private/Factories/EditorFactories.cpp
@@ -6221,6 +6221,7 @@ EReimportResult::Type UReimportFbxSkeletalMeshFactory::Reimport( UObject* Obj, i
 		ImportOptions->SkeletonForAnimation = SkeletalMesh->Skeleton;
 		ImportOptions->bCreatePhysicsAsset = false;
 		ImportOptions->PhysicsAsset = SkeletalMesh->PhysicsAsset;
+		
 
 		ImportOptions = GetImportOptions( FFbxImporter, ReimportUI, bShowOptionDialog, bIsAutomated, Obj->GetPathName(), bOperationCanceled, bOutImportAll, bIsObjFormat, Filename, bForceImportType, FBXIT_SkeletalMesh);
 
@@ -6231,11 +6232,32 @@ EReimportResult::Type UReimportFbxSkeletalMeshFactory::Reimport( UObject* Obj, i
 		}
 	}
 
+	//Set the build option to reflect the user choice in the dialog
+	UFbxSkeletalMeshImportData* SKImportData = UFbxSkeletalMeshImportData::GetImportDataForSkeletalMesh(SkeletalMesh, ReimportUI->SkeletalMeshImportData);
+	if (SKImportData)
+	{
+		FSkeletalMeshLODInfo* LODInfo = SkeletalMesh->GetLODInfo(0);
+		if (LODInfo && SkeletalMesh->GetImportedModel() && SkeletalMesh->GetImportedModel()->LODModels.IsValidIndex(0))
+		{
+			const FSkeletalMeshLODModel& LODModel = SkeletalMesh->GetImportedModel()->LODModels[0];
+			if (LODModel.RawSkeletalMeshBulkData.IsBuildDataAvailable())
+			{
+				//Set the build settings
+				LODInfo->BuildSettings.bComputeWeightedNormals = SKImportData->bComputeWeightedNormals;
+				LODInfo->BuildSettings.bRecomputeNormals = SKImportData->NormalImportMethod == EFBXNormalImportMethod::FBXNIM_ComputeNormals;
+				LODInfo->BuildSettings.bRecomputeTangents = SKImportData->NormalImportMethod != EFBXNormalImportMethod::FBXNIM_ImportNormalsAndTangents;
+				LODInfo->BuildSettings.bUseMikkTSpace = SKImportData->NormalGenerationMethod == EFBXNormalGenerationMethod::MikkTSpace;
+			}
+		}
+	}
+
 	UE_LOG(LogEditorFactories, Log, TEXT("Performing atomic reimport of [%s]"), *Filename);
 	CurrentFilename = Filename;
 
 	if( !bOperationCanceled )
 	{
+		FScopedSkeletalMeshPostEditChange ScopedPostEditChange(SkeletalMesh);
+
 		ImportOptions->bCanShowDialog = !IsUnattended;
 
 		if (ImportOptions->bImportAsSkeletalSkinning)
@@ -6293,7 +6315,7 @@ EReimportResult::Type UReimportFbxSkeletalMeshFactory::Reimport( UObject* Obj, i
 			//Restore skin weight profile infos, then reimport affected LODs
 			TArray<FSkinWeightProfileInfo>&SkinWeightsProfile = SkeletalMesh->GetSkinWeightProfiles();
 			SkinWeightsProfile = ExistingSkinWeightProfileInfos;
-			FSkinWeightsUtilities::ReimportAlternateSkinWeight(SkeletalMesh, 0, true);
+			FSkinWeightsUtilities::ReimportAlternateSkinWeight(SkeletalMesh, 0);
 		}
 
 		// Reimporting can have dangerous effects if the mesh is still in the transaction buffer.  Reset the transaction buffer if this is the case
diff --git a/Engine/Source/Editor/UnrealEd/Private/Factories/SkeletalMeshImport.cpp b/Engine/Source/Editor/UnrealEd/Private/Factories/SkeletalMeshImport.cpp
--- a/Engine/Source/Editor/UnrealEd/Private/Factories/SkeletalMeshImport.cpp
+++ b/Engine/Source/Editor/UnrealEd/Private/Factories/SkeletalMeshImport.cpp
@@ -255,17 +255,20 @@ ExistingSkelMeshData* SaveExistingSkelMeshData(USkeletalMesh* ExistingSkelMesh,
 
 	for (int32 LodIndex = 0; LodIndex < ImportedResource->LODModels.Num(); ++LodIndex)
 	{
+		const FSkeletalMeshLODModel& LODModel = ImportedResource->LODModels[LodIndex];
 		ExistingMeshDataPtr->ExistingImportMeshLodSectionMaterialData.AddZeroed();
-		for (int32 SectionIndex = 0; SectionIndex < ImportedResource->LODModels[LodIndex].Sections.Num(); ++SectionIndex)
+		for (int32 SectionIndex = 0; SectionIndex < LODModel.Sections.Num(); ++SectionIndex)
 		{
-			int32 SectionMaterialIndex = ImportedResource->LODModels[LodIndex].Sections[SectionIndex].MaterialIndex;
-			bool SectionCastShadow = ImportedResource->LODModels[LodIndex].Sections[SectionIndex].bCastShadow;
-			bool SectionRecomputeTangents = ImportedResource->LODModels[LodIndex].Sections[SectionIndex].bRecomputeTangent;
-			int32 GenerateUpTo = ImportedResource->LODModels[LodIndex].Sections[SectionIndex].GenerateUpToLodIndex;
-			bool bDisabled = ImportedResource->LODModels[LodIndex].Sections[SectionIndex].bDisabled;
+			const FSkelMeshSourceSectionUserData& UserSectionData = LODModel.UserSectionsData.FindChecked(LODModel.Sections[SectionIndex].OriginalDataSectionIndex);
+			int32 SectionMaterialIndex = LODModel.Sections[SectionIndex].MaterialIndex;
+			bool SectionCastShadow = UserSectionData.bCastShadow;
+			bool SectionRecomputeTangents = UserSectionData.bRecomputeTangent;
+			int32 GenerateUpTo = UserSectionData.GenerateUpToLodIndex;
+			bool bDisabled = UserSectionData.bDisabled;
+			bool bBoneChunkedSection = LODModel.Sections[SectionIndex].ChunkedParentSectionIndex != INDEX_NONE;
 			if (ExistingMeshDataPtr->ExistingImportMaterialOriginalNameData.IsValidIndex(SectionMaterialIndex))
 			{
-				ExistingMeshDataPtr->ExistingImportMeshLodSectionMaterialData[LodIndex].Add(ExistingMeshLodSectionData(ExistingMeshDataPtr->ExistingImportMaterialOriginalNameData[SectionMaterialIndex], SectionCastShadow, SectionRecomputeTangents, GenerateUpTo, bDisabled));
+				ExistingMeshDataPtr->ExistingImportMeshLodSectionMaterialData[LodIndex].Add(ExistingMeshLodSectionData(ExistingMeshDataPtr->ExistingImportMaterialOriginalNameData[SectionMaterialIndex], SectionCastShadow, SectionRecomputeTangents, GenerateUpTo, bDisabled, bBoneChunkedSection));
 			}
 		}
 	}
@@ -289,6 +292,7 @@ ExistingSkelMeshData* SaveExistingSkelMeshData(USkeletalMesh* ExistingSkelMesh,
 			//Save the imported LOD reduction settings
 			ExistingMeshDataPtr->ExistingReimportLODReductionSettings = LODInfo->ReductionSettings;
 		}
+		ExistingMeshDataPtr->ExistingBaseLODInfo = *LODInfo;
 
 		// Remove the zero'th LOD (ie: the LOD being reimported).
 		if (!ReimportSpecificLOD)
@@ -310,9 +314,9 @@ ExistingSkelMeshData* SaveExistingSkelMeshData(USkeletalMesh* ExistingSkelMesh,
 			{
 				FSkeletalMeshLODModel BaseLODModel;
 				TMap<FString, TArray<FMorphTargetDelta>> BaseLODMorphTargetData;
-				ImportedResource->OriginalReductionSourceMeshData[ReductionLODIndex]->LoadReductionData(BaseLODModel, BaseLODMorphTargetData);
+				ImportedResource->OriginalReductionSourceMeshData[ReductionLODIndex]->LoadReductionData(BaseLODModel, BaseLODMorphTargetData, ExistingSkelMesh);
 				FReductionBaseSkeletalMeshBulkData* ReductionLODData = new FReductionBaseSkeletalMeshBulkData();
-				ReductionLODData->SaveReductionData(BaseLODModel, BaseLODMorphTargetData);
+				ReductionLODData->SaveReductionData(BaseLODModel, BaseLODMorphTargetData, ExistingSkelMesh);
 				//Add necessary empty slot
 				while (ExistingMeshDataPtr->ExistingOriginalReductionSourceMeshData.Num() < LODModelIndex)
 				{
@@ -361,9 +365,6 @@ ExistingSkelMeshData* SaveExistingSkelMeshData(USkeletalMesh* ExistingSkelMesh,
 	ExistingMeshDataPtr->ExistingMorphTargets.Empty(ExistingSkelMesh->MorphTargets.Num());
 	ExistingMeshDataPtr->ExistingMorphTargets.Append(ExistingSkelMesh->MorphTargets);
 	
-	ExistingMeshDataPtr->bExistingUseFullPrecisionUVs = ExistingSkelMesh->bUseFullPrecisionUVs;
-	ExistingMeshDataPtr->bExistingUseHighPrecisionTangentBasis = ExistingSkelMesh->bUseHighPrecisionTangentBasis;
-
 	ExistingMeshDataPtr->ExistingAssetImportData = ExistingSkelMesh->AssetImportData;
 	ExistingMeshDataPtr->ExistingThumbnailInfo = ExistingSkelMesh->ThumbnailInfo;
 
@@ -442,6 +443,31 @@ void RegenerateDependentLODs(USkeletalMesh* SkelMesh, int32 LODIndex)
 	}
 }
 
+void RestoreDependentLODs(ExistingSkelMeshData* MeshData, USkeletalMesh* SkeletalMesh)
+{
+	check(SkeletalMesh != nullptr);
+	int32 TotalLOD = MeshData->ExistingLODModels.Num();
+	FSkeletalMeshModel* SkeletalMeshImportedModel = SkeletalMesh->GetImportedModel();
+
+	for (int32 Index = 0; Index < TotalLOD; ++Index)
+	{
+		int32 LODIndex = Index + 1;
+		if (LODIndex >= SkeletalMesh->GetLODInfoArray().Num())
+		{
+			FSkeletalMeshLODInfo& ExistLODInfo = MeshData->ExistingLODInfo[Index];
+			FSkeletalMeshLODModel& ExistLODModel = MeshData->ExistingLODModels[Index];
+			// reset material maps, it won't work anyway. 
+			ExistLODInfo.LODMaterialMap.Empty();
+
+			FSkeletalMeshLODModel* NewLODModel = new FSkeletalMeshLODModel(ExistLODModel);
+			SkeletalMeshImportedModel->LODModels.Add(NewLODModel);
+			// add LOD info back
+			SkeletalMesh->AddLODInfo(ExistLODInfo);
+			check(LODIndex < SkeletalMesh->GetLODInfoArray().Num());
+		}
+	}
+}
+
 void TryRegenerateLODs(ExistingSkelMeshData* MeshData, USkeletalMesh* SkeletalMesh)
 {
 	check(SkeletalMesh != nullptr);
@@ -710,24 +736,15 @@ void RestoreExistingSkelMeshData(ExistingSkelMeshData* MeshData, USkeletalMesh*
 	}
 
 	//Copy back the reimport LOD specific data
-	if (SkeletalMesh->GetLODInfoArray().IsValidIndex(ReimportLODIndex) && MeshData->bIsReimportLODReduced)
+	if (SkeletalMesh->GetLODInfoArray().IsValidIndex(ReimportLODIndex))
 	{
 		FSkeletalMeshLODInfo& BaseLODInfo = SkeletalMesh->GetLODInfoArray()[ReimportLODIndex];
-		//Restore the reimport LOD reduction settings
-		BaseLODInfo.ReductionSettings = MeshData->ExistingReimportLODReductionSettings;
-		if (SkeletalMeshImportedModel && SkeletalMeshImportedModel->OriginalReductionSourceMeshData.IsValidIndex(ReimportLODIndex))
+		//Restore the build setting first
+		BaseLODInfo.BuildSettings = MeshData->ExistingBaseLODInfo.BuildSettings;
+		if (MeshData->bIsReimportLODReduced)
 		{
-			SkeletalMeshImportedModel->OriginalReductionSourceMeshData[ReimportLODIndex]->EmptyBulkData();
-		}
-
-		if (SkeletalMesh->IsReductionActive(ReimportLODIndex))
-		{
-			//Regenerate the reimport LOD
-			GWarn->BeginSlowTask(LOCTEXT("RegenReimportedLOD", "Generating reimported LOD"), true);
-			FSkeletalMeshUpdateContext UpdateContext;
-			UpdateContext.SkeletalMesh = SkeletalMesh;
-			FLODUtilities::SimplifySkeletalMeshLOD(UpdateContext, ReimportLODIndex, false);
-			GWarn->EndSlowTask();
+			//Restore the reimport LOD reduction settings
+			BaseLODInfo.ReductionSettings = MeshData->ExistingBaseLODInfo.ReductionSettings;
 		}
 	}
 
@@ -755,7 +772,7 @@ void RestoreExistingSkelMeshData(ExistingSkelMeshData* MeshData, USkeletalMesh*
 		// We copy back and fix-up the LODs that still work with this skeleton.
 		if (MeshData->ExistingLODModels.Num() > 0)
 		{
-			auto RestoreReductionSourceData = [&SkeletalMeshImportedModel, &MeshData](int32 ExistingIndex, int32 NewIndex)
+			auto RestoreReductionSourceData = [&SkeletalMesh, &SkeletalMeshImportedModel, &MeshData](int32 ExistingIndex, int32 NewIndex)
 			{
 				if (!MeshData->ExistingOriginalReductionSourceMeshData.IsValidIndex(ExistingIndex) || MeshData->ExistingOriginalReductionSourceMeshData[ExistingIndex]->IsEmpty())
 				{
@@ -764,9 +781,9 @@ void RestoreExistingSkelMeshData(ExistingSkelMeshData* MeshData, USkeletalMesh*
 				//Restore the original reduction source mesh data
 				FSkeletalMeshLODModel BaseLODModel;
 				TMap<FString, TArray<FMorphTargetDelta>> BaseLODMorphTargetData;
-				MeshData->ExistingOriginalReductionSourceMeshData[ExistingIndex]->LoadReductionData(BaseLODModel, BaseLODMorphTargetData);
+				MeshData->ExistingOriginalReductionSourceMeshData[ExistingIndex]->LoadReductionData(BaseLODModel, BaseLODMorphTargetData, SkeletalMesh);
 				FReductionBaseSkeletalMeshBulkData* ReductionLODData = new FReductionBaseSkeletalMeshBulkData();
-				ReductionLODData->SaveReductionData(BaseLODModel, BaseLODMorphTargetData);
+				ReductionLODData->SaveReductionData(BaseLODModel, BaseLODMorphTargetData, SkeletalMesh);
 				//Add necessary empty slot
 				while (SkeletalMeshImportedModel->OriginalReductionSourceMeshData.Num() < NewIndex)
 				{
@@ -776,7 +793,6 @@ void RestoreExistingSkelMeshData(ExistingSkelMeshData* MeshData, USkeletalMesh*
 				SkeletalMeshImportedModel->OriginalReductionSourceMeshData.Add(ReductionLODData);
 			};
 
-
 			if (SkeletonsAreCompatible(SkeletalMesh->RefSkeleton, MeshData->ExistingRefSkeleton, bImportSkinningOnly))
 			{
 				// First create mapping table from old skeleton to new skeleton.
@@ -889,8 +905,8 @@ void RestoreExistingSkelMeshData(ExistingSkelMeshData* MeshData, USkeletalMesh*
 					}
 				}
 			}
-			//Always regenerate LODs when re-importing LOD 0
-			TryRegenerateLODs(MeshData, SkeletalMesh);
+			//We just need to restore the LOD model and LOD info the build will regenerate the LODs
+			RestoreDependentLODs(MeshData, SkeletalMesh);
 		}
 
 		for (int32 AssetIndex = 0; AssetIndex < MeshData->ExistingPhysicsAssets.Num(); ++AssetIndex)
@@ -920,9 +936,6 @@ void RestoreExistingSkelMeshData(ExistingSkelMeshData* MeshData, USkeletalMesh*
 		SkeletalMesh->MorphTargets.Append(MeshData->ExistingMorphTargets);
 		SkeletalMesh->InitMorphTargets();
 
-		SkeletalMesh->bUseFullPrecisionUVs = MeshData->bExistingUseFullPrecisionUVs;
-		SkeletalMesh->bUseHighPrecisionTangentBasis= MeshData->bExistingUseHighPrecisionTangentBasis;
-
 		SkeletalMesh->AssetImportData = MeshData->ExistingAssetImportData.Get();
 		SkeletalMesh->ThumbnailInfo = MeshData->ExistingThumbnailInfo.Get();
 
@@ -936,15 +949,39 @@ void RestoreExistingSkelMeshData(ExistingSkelMeshData* MeshData, USkeletalMesh*
 
 		SkeletalMesh->SetSamplingInfo(MeshData->ExistingSamplingInfo);
 	}
-	else
-	{
-		RegenerateDependentLODs(SkeletalMesh, ReimportLODIndex);
-	}
 
 	//Restore the section change only for the reimport LOD, other LOD are not affected since the material array can only grow.
 	if (MeshData->UseMaterialNameSlotWorkflow)
 	{
+		//Restore the base LOD materialMap the LODs LODMaterialMap are restore differently
+		if (ReimportLODIndex < 1 && SkeletalMesh->GetLODInfoArray().IsValidIndex(ReimportLODIndex))
+		{
+			FSkeletalMeshLODInfo& BaseLODInfo = SkeletalMesh->GetLODInfoArray()[ReimportLODIndex];
+			//Restore the Base MaterialMap
+			if (SkeletalMesh->GetImportedModel() && SkeletalMesh->GetImportedModel()->LODModels.IsValidIndex(ReimportLODIndex))
+			{
+				FSkeletalMeshLODModel& LODModel = SkeletalMesh->GetImportedModel()->LODModels[ReimportLODIndex];
+				for (int32 SectionIndex = 0; SectionIndex < LODModel.Sections.Num(); ++SectionIndex)
+				{
+					int32 MaterialIndex = LODModel.Sections[SectionIndex].MaterialIndex;
+					if (MeshData->ExistingBaseLODInfo.LODMaterialMap.IsValidIndex(MaterialIndex))
+					{
+						int32 ExistingLODMaterialIndex = MeshData->ExistingBaseLODInfo.LODMaterialMap[MaterialIndex];
+						while (BaseLODInfo.LODMaterialMap.Num() <= MaterialIndex)
+						{
+							BaseLODInfo.LODMaterialMap.Add(INDEX_NONE);
+						}
+						if (RemapMaterial.IsValidIndex(ExistingLODMaterialIndex))
+						{
+							ExistingLODMaterialIndex = RemapMaterial[ExistingLODMaterialIndex];
+						}
+						BaseLODInfo.LODMaterialMap[MaterialIndex] = ExistingLODMaterialIndex;
+					}
+				}
+			}
+		}
 		FSkeletalMeshLODModel &NewSkelMeshLodModel = SkeletalMeshImportedModel->LODModels[ReimportLODIndex];
+		int32 ParentSectionIndex = 0;
 		//Restore the section changes from the old import data
 		for (int32 SectionIndex = 0; SectionIndex < NewSkelMeshLodModel.Sections.Num(); SectionIndex++)
 		{
@@ -963,6 +1000,20 @@ void RestoreExistingSkelMeshData(ExistingSkelMeshData* MeshData, USkeletalMesh*
 						NewSkelMeshLodModel.Sections[SectionIndex].bRecomputeTangent = MeshData->ExistingImportMeshLodSectionMaterialData[ReimportLODIndex][RemapMaterial[NewMeshSectionMaterialIndex]].bRecomputeTangents;
 						NewSkelMeshLodModel.Sections[SectionIndex].GenerateUpToLodIndex = MeshData->ExistingImportMeshLodSectionMaterialData[ReimportLODIndex][RemapMaterial[NewMeshSectionMaterialIndex]].GenerateUpTo;
 						NewSkelMeshLodModel.Sections[SectionIndex].bDisabled = MeshData->ExistingImportMeshLodSectionMaterialData[ReimportLODIndex][RemapMaterial[NewMeshSectionMaterialIndex]].bDisabled;
+						bool bBoneChunkedSection = MeshData->ExistingImportMeshLodSectionMaterialData[ReimportLODIndex][RemapMaterial[NewMeshSectionMaterialIndex]].bBoneChunkedSection;
+						if (!bBoneChunkedSection)
+						{
+							//Set the new Parent Index
+							ParentSectionIndex = SectionIndex;
+							FSkelMeshSourceSectionUserData& UserSectionData = NewSkelMeshLodModel.UserSectionsData.FindOrAdd(ParentSectionIndex);
+							UserSectionData.bDisabled = NewSkelMeshLodModel.Sections[SectionIndex].bDisabled;
+							UserSectionData.bCastShadow = NewSkelMeshLodModel.Sections[SectionIndex].bCastShadow;
+							UserSectionData.bRecomputeTangent = NewSkelMeshLodModel.Sections[SectionIndex].bRecomputeTangent;
+							UserSectionData.GenerateUpToLodIndex = NewSkelMeshLodModel.Sections[SectionIndex].GenerateUpToLodIndex;
+							//The cloth will be rebind later after the reimport is done
+						}
+						NewSkelMeshLodModel.Sections[SectionIndex].ChunkedParentSectionIndex = bBoneChunkedSection ? ParentSectionIndex : INDEX_NONE;
+						NewSkelMeshLodModel.Sections[SectionIndex].OriginalDataSectionIndex = ParentSectionIndex;
 					}
 				}
 			}
@@ -1007,11 +1058,12 @@ void RestoreExistingSkelMeshData(ExistingSkelMeshData* MeshData, USkeletalMesh*
 				break;
 			}
 		}
+		//Make sure we reset the User section array to only what we have in the fbx
+		NewSkelMeshLodModel.SyncronizeUserSectionsDataArray(true);
 	}
 
 	if (!bImportSkinningOnly && !MeshData->bIsReimportLODReduced)
 	{
-		//Fix the OriginalReductionSourceMeshData
 		if (ReimportLODIndex < 0)
 		{
 			if (SkeletalMeshImportedModel->OriginalReductionSourceMeshData.IsValidIndex(0))
diff --git a/Engine/Source/Editor/UnrealEd/Private/Fbx/FbxFactory.cpp b/Engine/Source/Editor/UnrealEd/Private/Fbx/FbxFactory.cpp
--- a/Engine/Source/Editor/UnrealEd/Private/Fbx/FbxFactory.cpp
+++ b/Engine/Source/Editor/UnrealEd/Private/Fbx/FbxFactory.cpp
@@ -549,6 +549,10 @@ UObject* UFbxFactory::FactoryCreateFile
 						int32 LODIndex;
 						int32 SuccessfulLodIndex = 0;
 						bool bImportSkeletalMeshLODs = ImportUI->SkeletalMeshImportData->bImportMeshLODs;
+						
+						//The skeletalmesh will be set after we import the LOD 0 since it is not created yet.
+						FScopedSkeletalMeshPostEditChange ScopedPostEditChange(nullptr);
+
 						for (LODIndex = 0; LODIndex < MaxLODLevel; LODIndex++)
 						{
 							//We need to know what is the imported lod index when importing the morph targets
@@ -605,7 +609,7 @@ UObject* UFbxFactory::FactoryCreateFile
 
 								USkeletalMesh* NewMesh = FbxImporter->ImportSkeletalMesh( ImportSkeletalMeshArgs );
 								CreatedObject = NewMesh;
-
+								
 								if(bOperationCanceled)
 								{
 									// User cancelled, clean up and return
@@ -617,6 +621,9 @@ UObject* UFbxFactory::FactoryCreateFile
 
 								if ( NewMesh )
 								{
+									//Set the base skeletalmesh to the scoped post edit change variable
+									ScopedPostEditChange.SetSkeletalMesh(NewMesh);
+
 									if (ImportOptions->bImportAnimations)
 									{
 										// We need to remove all scaling from the root node before we set up animation data.
@@ -643,7 +650,7 @@ UObject* UFbxFactory::FactoryCreateFile
 								LODInfo.ReductionSettings.BaseLOD = 0;
 								LODInfo.bImportWithBaseMesh = true;
 								LODInfo.SourceImportFilename = FString(TEXT(""));
-								FLODUtilities::SimplifySkeletalMeshLOD(UpdateContext, SuccessfulLodIndex, false);
+								FLODUtilities::SimplifySkeletalMeshLOD(UpdateContext, SuccessfulLodIndex);
 								ImportedSuccessfulLodIndex = SuccessfulLodIndex;
 								SuccessfulLodIndex++;
 							}
@@ -662,7 +669,7 @@ UObject* UFbxFactory::FactoryCreateFile
 								ImportSkeletalMeshArgs.OutData = &OutData;
 
 								USkeletalMesh *LODObject = FbxImporter->ImportSkeletalMesh( ImportSkeletalMeshArgs );
-								bool bImportSucceeded = !bOperationCanceled && FbxImporter->ImportSkeletalMeshLOD(LODObject, BaseSkeletalMesh, SuccessfulLodIndex, false);
+								bool bImportSucceeded = !bOperationCanceled && FbxImporter->ImportSkeletalMeshLOD(LODObject, BaseSkeletalMesh, SuccessfulLodIndex);
 
 								if (bImportSucceeded)
 								{
@@ -704,15 +711,6 @@ UObject* UFbxFactory::FactoryCreateFile
 							
 							USkeletalMesh* SkeletalMesh = Cast<USkeletalMesh>(CreatedObject);
 							UnFbx::FFbxImporter::UpdateSkeletalMeshImportData(SkeletalMesh, ImportUI->SkeletalMeshImportData, INDEX_NONE, nullptr, nullptr);
-							
-							//If we have import some morph target we have to rebuild the render resources since morph target are now using GPU
-							if (SkeletalMesh)// && SkeletalMesh->MorphTargets.Num() > 0)
-							{
-								SkeletalMesh->ReleaseResources();
-								//Rebuild the resources with a post edit change since we have added some morph targets
-								SkeletalMesh->PostEditChange();
-							}
-
 						}
 					}
 				
@@ -1391,7 +1389,7 @@ namespace ImportCompareHelper
 		}
 		else if (!bCombineMeshes && !bCombineMeshesLOD)
 		{
-			Node = FFbxImporter->GetMeshNodesFromName(StaticMesh, FbxMeshArray);
+			Node = FFbxImporter->GetMeshNodesFromName(StaticMesh->GetName(), FbxMeshArray);
 		}
 
 		// If there is no match it may be because an LOD group was imported where
@@ -1481,7 +1479,38 @@ namespace ImportCompareHelper
 			return;
 		}
 
-		const TArray<FbxNode*>& SkeletalMeshNodes = *(SkeletalMeshArray[0]);
+		TArray<FbxNode*>* FbxNodes = nullptr;
+		// if there is only one mesh, use it without name checking 
+		// (because the "Used As Full Name" option enables users name the Unreal mesh by themselves
+		if (SkeletalMeshArray.Num() == 1)
+		{
+			FbxNodes = SkeletalMeshArray[0];
+		}
+		else if(SkeletalMeshArray.Num() > 1)
+		{
+			FbxNodes = nullptr;
+			//Search a set corresponding
+			for (TArray<FbxNode*>*SkeletalMeshNodes : SkeletalMeshArray)
+			{
+				FbxNode* Node = FFbxImporter->GetMeshNodesFromName(SkeletalMesh->GetName(), *SkeletalMeshNodes);
+				if (Node != nullptr)
+				{
+					FbxNodes = SkeletalMeshNodes;
+					break;
+				}
+			}
+			if (FbxNodes == nullptr)
+			{
+				//Re import the first skeletalmesh found in the fbx file
+				FbxNodes = SkeletalMeshArray[0];
+			}
+		}
+		if (FbxNodes == nullptr)
+		{
+			return;
+		}
+
+		const TArray<FbxNode*>& SkeletalMeshNodes = *(FbxNodes);
 		if (SkeletalMeshNodes.Num() == 0)
 		{
 			return;
diff --git a/Engine/Source/Editor/UnrealEd/Private/Fbx/FbxMainImport.cpp b/Engine/Source/Editor/UnrealEd/Private/Fbx/FbxMainImport.cpp
--- a/Engine/Source/Editor/UnrealEd/Private/Fbx/FbxMainImport.cpp
+++ b/Engine/Source/Editor/UnrealEd/Private/Fbx/FbxMainImport.cpp
@@ -1741,6 +1741,52 @@ FName FFbxImporter::MakeNameForMesh(FString InName, FbxObject* FbxObject)
 	return OutputName;
 }
 
+FbxNode* FFbxImporter::GetMeshNodesFromName(const FString& ReimportMeshName, TArray<FbxNode*>& FbxMeshArray)
+{
+	//StaticMesh->GetName()
+	char MeshName[1024];
+	FCStringAnsi::Strcpy(MeshName, 1024, TCHAR_TO_UTF8(*ReimportMeshName));
+	// find the Fbx mesh node that the Unreal Mesh matches according to name
+	int32 MeshIndex;
+	for (MeshIndex = 0; MeshIndex < FbxMeshArray.Num(); MeshIndex++)
+	{
+		const char* FbxMeshName = FbxMeshArray[MeshIndex]->GetName();
+		// The name of Unreal mesh may have a prefix, so we match from end
+		int32 i = 0;
+		char* MeshPtr = MeshName + FCStringAnsi::Strlen(MeshName) - 1;
+		if (FCStringAnsi::Strlen(FbxMeshName) <= FCStringAnsi::Strlen(MeshName))
+		{
+			const char* FbxMeshPtr = FbxMeshName + FCStringAnsi::Strlen(FbxMeshName) - 1;
+			while (i < FCStringAnsi::Strlen(FbxMeshName))
+			{
+				bool bIsPointAndUnderscore = *FbxMeshPtr == '.' && *MeshPtr == '_';
+
+				if (*MeshPtr != *FbxMeshPtr && !bIsPointAndUnderscore)
+				{
+					break;
+				}
+				else
+				{
+					i++;
+					MeshPtr--;
+					FbxMeshPtr--;
+				}
+			}
+		}
+
+		if (i == FCStringAnsi::Strlen(FbxMeshName)) // matched
+		{
+			// check further
+			if (FCStringAnsi::Strlen(FbxMeshName) == FCStringAnsi::Strlen(MeshName) || // the name of Unreal mesh is full match
+				*MeshPtr == '_')														// or the name of Unreal mesh has a prefix
+			{
+				return FbxMeshArray[MeshIndex];
+			}
+		}
+	}
+	return nullptr;
+}
+
 FbxAMatrix FFbxImporter::ComputeSkeletalMeshTotalMatrix(FbxNode* Node, FbxNode *RootSkeletalNode)
 {
 	if (ImportOptions->bImportScene && !ImportOptions->bTransformVertexToAbsolute && RootSkeletalNode != nullptr && RootSkeletalNode != Node)
diff --git a/Engine/Source/Editor/UnrealEd/Private/Fbx/FbxSceneImportFactory.cpp b/Engine/Source/Editor/UnrealEd/Private/Fbx/FbxSceneImportFactory.cpp
--- a/Engine/Source/Editor/UnrealEd/Private/Fbx/FbxSceneImportFactory.cpp
+++ b/Engine/Source/Editor/UnrealEd/Private/Fbx/FbxSceneImportFactory.cpp
@@ -58,6 +58,7 @@
 
 #include "HAL/FileManager.h"
 #include "LODUtilities.h"
+#include "ComponentReregisterContext.h"
 
 #define LOCTEXT_NAMESPACE "FBXSceneImportFactory"
 
@@ -1886,6 +1887,9 @@ UObject* UFbxSceneImportFactory::ImportOneSkeletalMesh(void* VoidRootNodeToImpor
 		}
 	}
 
+	//The skeletalmesh will be set after we import the LOD 0 since it is not created yet.
+	FScopedSkeletalMeshPostEditChange ScopedPostEditChange(nullptr);
+
 	int32 LODIndex;
 	for (LODIndex = 0; LODIndex < MaxLODLevel; LODIndex++)
 	{
@@ -1969,6 +1973,7 @@ UObject* UFbxSceneImportFactory::ImportOneSkeletalMesh(void* VoidRootNodeToImpor
 			NewObject = NewMesh;
 			if (NewMesh)
 			{
+				ScopedPostEditChange.SetSkeletalMesh(NewMesh);
 				TSharedPtr<FFbxNodeInfo> SkelMeshNodeInfo;
 				if (FindSceneNodeInfo(SceneInfo, SkelMeshNodeArray[0]->GetUniqueID(), SkelMeshNodeInfo) && SkelMeshNodeInfo.IsValid() && SkelMeshNodeInfo->AttributeInfo.IsValid())
 				{
@@ -2046,17 +2051,10 @@ UObject* UFbxSceneImportFactory::ImportOneSkeletalMesh(void* VoidRootNodeToImpor
 		}
 	}
 	
-	USkeletalMesh* ImportedSkelMesh = Cast<USkeletalMesh>(NewObject);
-	//If we have import some morph target we have to rebuild the render resources since morph target are now using GPU
-	if (ImportedSkelMesh && ImportedSkelMesh->MorphTargets.Num() > 0)
-	{
-		ImportedSkelMesh->ReleaseResources();
-		//Rebuild the resources with a post edit change since we have added some morph targets
-		ImportedSkelMesh->PostEditChange();
-	}
-	
 	//Put back the options
 	GlobalImportSettings->bBakePivotInVertex = Old_bBakePivotInVertex;
+
+	//FScopedSkeletalMeshPostEditChange will call post edit change when going out of scope
 	return NewObject;
 }
 
