Merging //UE4/Main@9388041 to Dev-RenderPlat-Staging

diff --git a/Engine/Source/Editor/Persona/Private/AnimGraphNodeDetails.cpp b/Engine/Source/Editor/Persona/Private/AnimGraphNodeDetails.cpp
--- a/Engine/Source/Editor/Persona/Private/AnimGraphNodeDetails.cpp
+++ b/Engine/Source/Editor/Persona/Private/AnimGraphNodeDetails.cpp
@@ -44,6 +44,8 @@
 #include "Engine/SkeletalMeshSocket.h"
 #include "Styling/CoreStyle.h"
 #include "Widgets/Input/SCheckBox.h"
+#include "LODInfoUILayout.h"
+#include "IPersonaToolkit.h"
 
 #define LOCTEXT_NAMESPACE "KismetNodeWithOptionalPinsDetails"
 
@@ -169,8 +171,6 @@ void FAnimGraphNodeDetails::CustomizeDetails(class IDetailLayoutBuilder& DetailB
 
 						ShowHidePropertyHandle->MarkHiddenByCustomization();
 
-						TSharedRef<SWidget> ShowHidePropertyWidget = CreateAsPinWidget(ShowHidePropertyHandle);
-
 						ValueWidget->SetVisibility(TAttribute<EVisibility>::Create(TAttribute<EVisibility>::FGetter::CreateSP(this, &FAnimGraphNodeDetails::GetVisibilityOfProperty, ShowHidePropertyHandle)));
 
 						// If we have an edit condition, that comes as part of the default name widget, so just use a text block to avoid duplicate checkboxes
@@ -187,26 +187,7 @@ void FAnimGraphNodeDetails::CustomizeDetails(class IDetailLayoutBuilder& DetailB
 							PropertyNameWidget = NameWidget;
 						}
 
-						NameWidget = 
-							SNew(SHorizontalBox)
-							+SHorizontalBox::Slot()
-							.HAlign(HAlign_Left)
-							.VAlign(VAlign_Center)
-							.FillWidth(1)
-							[
-								SNew(SBox)
-								.Clipping(EWidgetClipping::ClipToBounds)
-								[
-									PropertyNameWidget.ToSharedRef()
-								]
-							]
-							+SHorizontalBox::Slot()
-							.AutoWidth()
-							.HAlign(HAlign_Right)
-							.VAlign(VAlign_Center)
-							[
-								ShowHidePropertyWidget
-							];
+						NameWidget = PropertyNameWidget;
 
 						// we only show children if visibility is one
 						// whenever toggles, this gets called, so it will be refreshed
@@ -250,68 +231,6 @@ void FAnimGraphNodeDetails::CustomizeDetails(class IDetailLayoutBuilder& DetailB
 	}
 }
 
-TSharedRef<SWidget> FAnimGraphNodeDetails::CreateAsPinWidget(TSharedRef<IPropertyHandle> InPropertyHandle)
-{
-	TWeakPtr<IPropertyHandle> WeakPropertyHandle = InPropertyHandle;
-
-	auto IsCheckedLambda = [WeakPropertyHandle]()
-	{
-		if(WeakPropertyHandle.IsValid())
-		{
-			bool bValue;
-			FPropertyAccess::Result Result = WeakPropertyHandle.Pin()->GetValue(bValue);
-			if(Result == FPropertyAccess::MultipleValues)
-			{
-				return ECheckBoxState::Undetermined;
-			}
-			else
-			{
-				return bValue ? ECheckBoxState::Checked : ECheckBoxState::Unchecked;
-			}
-		}
-
-		return ECheckBoxState::Unchecked;
-	};
-
-	auto OnCheckStateChangedLambda = [WeakPropertyHandle](ECheckBoxState InCheckBoxState)
-	{
-		if(WeakPropertyHandle.IsValid())
-		{
-			bool bValue = InCheckBoxState == ECheckBoxState::Checked;
-			WeakPropertyHandle.Pin()->SetValue(bValue);
-		}
-	};
-
-	auto ImageLambda = [WeakPropertyHandle]()
-	{
-		if(WeakPropertyHandle.IsValid())
-		{
-			bool bValue;
-			FPropertyAccess::Result Result = WeakPropertyHandle.Pin()->GetValue(bValue);
-			if(Result == FPropertyAccess::MultipleValues)
-			{
-				return FEditorStyle::GetBrush("Kismet.VariableList.HideForInstance");
-			}
-			else
-			{
-				return bValue ? FEditorStyle::GetBrush("Kismet.VariableList.ExposeForInstance") : FEditorStyle::GetBrush("Kismet.VariableList.HideForInstance");
-			}
-		}
-
-		return FEditorStyle::GetBrush("Kismet.VariableList.HideForInstance");
-	};
-
-	return SNew(SCheckBox)
-		.ToolTipText(LOCTEXT("AsPinTooltip", "Show/hide this property as a pin on the node"))
-		.IsChecked_Lambda(IsCheckedLambda)
-		.OnCheckStateChanged_Lambda(OnCheckStateChangedLambda)
-		.Style(FEditorStyle::Get(), "CheckboxLookToggleButtonCheckbox")
-		[
-			SNew(SImage)
-			.Image_Lambda(ImageLambda)
-			.ColorAndOpacity(FLinearColor::Black)
-		];
-}
 
 TSharedRef<SWidget> FAnimGraphNodeDetails::CreatePropertyWidget(UProperty* TargetProperty, TSharedRef<IPropertyHandle> TargetPropertyHandle, UClass* NodeClass)
 {
@@ -599,6 +518,7 @@ void FBoneReferenceCustomization::SetEditableSkeleton(TSharedRef<IPropertyHandle
 	UAnimGraphNode_Base* AnimGraphNode = nullptr;
 	USkeletalMesh* SkeletalMesh = nullptr;
 	UAnimationAsset * AnimationAsset = nullptr;
+	ULODInfoUILayout* LODInfoUILayout = nullptr;
 
 	USkeleton* TargetSkeleton = nullptr;
 	TSharedPtr<IEditableSkeleton> EditableSkeleton;
@@ -617,6 +537,16 @@ void FBoneReferenceCustomization::SetEditableSkeleton(TSharedRef<IPropertyHandle
 			TargetSkeleton = SkeletalMesh->Skeleton;
 			break;
 		}
+
+		LODInfoUILayout = Cast<ULODInfoUILayout>(*OuterIter);
+		if (LODInfoUILayout)
+		{
+			SkeletalMesh = LODInfoUILayout->GetPersonaToolkit()->GetPreviewMesh();
+			check(SkeletalMesh);
+			TargetSkeleton = SkeletalMesh->Skeleton;
+			break;
+		}
+
 		AnimationAsset = Cast<UAnimationAsset>(*OuterIter);
 		if (AnimationAsset)
 		{
@@ -711,10 +641,10 @@ TSharedRef<IPropertyTypeCustomization> FBoneSocketTargetCustomization::MakeInsta
 void FBoneSocketTargetCustomization::CustomizeChildren(TSharedRef<IPropertyHandle> StructPropertyHandle, class IDetailChildrenBuilder& ChildBuilder, IPropertyTypeCustomizationUtils& StructCustomizationUtils)
 {
 	// set property handle 
- 	SetPropertyHandle(StructPropertyHandle);
- 	// set editable skeleton info from struct
- 	SetEditableSkeleton(StructPropertyHandle);
- 	Build(StructPropertyHandle, ChildBuilder);
+	SetPropertyHandle(StructPropertyHandle);
+	// set editable skeleton info from struct
+	SetEditableSkeleton(StructPropertyHandle);
+	Build(StructPropertyHandle, ChildBuilder);
 }
 
 void FBoneSocketTargetCustomization::SetPropertyHandle(TSharedRef<IPropertyHandle> StructPropertyHandle)
@@ -1225,4 +1155,135 @@ void FPlayerTreeViewEntry::GenerateNameWidget(TSharedPtr<SHorizontalBox> Box)
 		];
 }
 
+void FAnimGraphNodeShowAsPinExtension::GetOptionalPinData(const IPropertyHandle& PropertyHandle, int32& OutOptionalPinIndex, UAnimGraphNode_Base*& OutAnimGraphNode) const
+{
+	OutOptionalPinIndex = INDEX_NONE;
+
+	TArray<UObject*> Objects;
+	PropertyHandle.GetOuterObjects(Objects);
+
+	UProperty* Property = PropertyHandle.GetProperty();
+	if(Property)
+	{
+		OutAnimGraphNode = Cast<UAnimGraphNode_Base>(Objects[0]);
+		if (OutAnimGraphNode != nullptr)
+		{
+			OutOptionalPinIndex = OutAnimGraphNode->ShowPinForProperties.IndexOfByPredicate([Property](const FOptionalPinFromProperty& InOptionalPin)
+			{
+				return Property->GetFName() == InOptionalPin.PropertyName;
+			});
+		}
+	}
+}
+
+bool FAnimGraphNodeShowAsPinExtension::IsPropertyExtendable(const UClass* InObjectClass, const IPropertyHandle& PropertyHandle) const
+{
+	int32 OptionalPinIndex;
+	UAnimGraphNode_Base* AnimGraphNode;
+	GetOptionalPinData(PropertyHandle, OptionalPinIndex, AnimGraphNode);
+
+	if(OptionalPinIndex != INDEX_NONE)
+	{
+		const FOptionalPinFromProperty& OptionalPin = AnimGraphNode->ShowPinForProperties[OptionalPinIndex];
+
+		// Not optional
+		if (!OptionalPin.bCanToggleVisibility && OptionalPin.bShowPin)
+		{
+			return false;
+		}
+
+		if(!PropertyHandle.GetProperty())
+		{
+			return false;
+		}
+
+		return OptionalPin.bCanToggleVisibility;
+	}
+
+	return false;
+}
+
+class SShowAsWidget : public SCompoundWidget
+{
+	SLATE_BEGIN_ARGS(SShowAsWidget) {}
+
+	SLATE_END_ARGS()
+
+	void Construct(const FArguments& InArgs, const TSharedRef<IPropertyHandle>& InPropertyHandle)
+	{
+		PropertyHandle = InPropertyHandle;
+
+		TSharedRef<SHorizontalBox> HorizontalBox = 
+			SNew(SHorizontalBox)
+			.ToolTipText(LOCTEXT("AsPinTooltip", "Show/hide this property as a pin on the node"));
+
+		TWeakPtr<SWidget> WeakHorizontalBox = HorizontalBox;
+
+		HorizontalBox->AddSlot()
+			.AutoWidth()
+			.VAlign(VAlign_Center)
+			[
+				SNew(STextBlock)
+				.Text(LOCTEXT("ExposeAsPinLabel", "Expose"))
+				.Font(IDetailLayoutBuilder::GetDetailFont())
+				.Visibility_Lambda([WeakHorizontalBox](){ return WeakHorizontalBox.IsValid() && WeakHorizontalBox.Pin()->IsHovered() ? EVisibility::Visible : EVisibility::Collapsed; })
+			];
+
+		HorizontalBox->AddSlot()
+			.AutoWidth()
+			.VAlign(VAlign_Center)
+			.Padding(3.0f, 0.0f, 0.0f, 0.0f)
+			[
+				SNew(SCheckBox)
+				.IsChecked(this, &SShowAsWidget::IsChecked)
+				.OnCheckStateChanged(this, &SShowAsWidget::OnCheckStateChanged)
+			];
+
+		ChildSlot
+		[
+			HorizontalBox
+		];
+	}
+
+	ECheckBoxState IsChecked() const
+	{
+		bool bValue;
+		FPropertyAccess::Result Result = PropertyHandle->GetValue(bValue);
+		if(Result == FPropertyAccess::MultipleValues)
+		{
+			return ECheckBoxState::Undetermined;
+		}
+		else
+		{
+			return bValue ? ECheckBoxState::Checked : ECheckBoxState::Unchecked;
+		}
+
+		return ECheckBoxState::Unchecked;
+	}
+
+	void OnCheckStateChanged(ECheckBoxState InCheckBoxState)
+	{
+		bool bValue = InCheckBoxState == ECheckBoxState::Checked;
+		PropertyHandle->SetValue(bValue);
+	}
+
+	TSharedPtr<IPropertyHandle> PropertyHandle;
+};
+
+TSharedRef<SWidget> FAnimGraphNodeShowAsPinExtension::GenerateExtensionWidget(const IDetailLayoutBuilder& InDetailBuilder, const UClass* InObjectClass, TSharedPtr<IPropertyHandle> PropertyHandle)
+{
+	int32 OptionalPinIndex;
+	UAnimGraphNode_Base* AnimGraphNode;
+	GetOptionalPinData(*PropertyHandle.Get(), OptionalPinIndex, AnimGraphNode);
+	check(OptionalPinIndex != INDEX_NONE);
+
+	const FName OptionalPinArrayEntryName(*FString::Printf(TEXT("ShowPinForProperties[%d].bShowPin"), OptionalPinIndex));
+	TSharedRef<IPropertyHandle> ShowHidePropertyHandle = InDetailBuilder.GetProperty(OptionalPinArrayEntryName, UAnimGraphNode_Base::StaticClass());
+
+	ShowHidePropertyHandle->MarkHiddenByCustomization();
+
+	return SNew(SShowAsWidget, ShowHidePropertyHandle);
+}
+
 #undef LOCTEXT_NAMESPACE
+
diff --git a/Engine/Source/Editor/Persona/Private/AnimGraphNodeDetails.h b/Engine/Source/Editor/Persona/Private/AnimGraphNodeDetails.h
--- a/Engine/Source/Editor/Persona/Private/AnimGraphNodeDetails.h
+++ b/Engine/Source/Editor/Persona/Private/AnimGraphNodeDetails.h
@@ -12,6 +12,7 @@
 #include "IDetailCustomization.h"
 #include "Widgets/Views/STableViewBase.h"
 #include "Widgets/Views/STableRow.h"
+#include "IDetailPropertyExtensionHandler.h"
 
 struct FAssetData;
 class FBlueprintEditor;
@@ -23,6 +24,18 @@ class USkeleton;
 class IEditableSkeleton;
 struct FAnimParentNodeAssetOverride;
 
+class FAnimGraphNodeShowAsPinExtension : public IDetailPropertyExtensionHandler
+{
+public:
+	// IDetailPropertyExtensionHandler interface
+	virtual bool IsPropertyExtendable(const UClass* InObjectClass, const IPropertyHandle& PropertyHandle) const override;
+	virtual TSharedRef<SWidget> GenerateExtensionWidget(const IDetailLayoutBuilder& InDetailBuilder, const UClass* InObjectClass, TSharedPtr<IPropertyHandle> PropertyHandle) override;
+
+private:
+	// Helper function
+	void GetOptionalPinData(const IPropertyHandle& PropertyHandle, int32& OutOptionalPinIndex, UAnimGraphNode_Base*& OutAnimGraphNode) const;
+};
+
 /////////////////////////////////////////////////////
 // FAnimGraphNodeDetails 
 
diff --git a/Engine/Source/Editor/Persona/Private/PersonaModule.cpp b/Engine/Source/Editor/Persona/Private/PersonaModule.cpp
--- a/Engine/Source/Editor/Persona/Private/PersonaModule.cpp
+++ b/Engine/Source/Editor/Persona/Private/PersonaModule.cpp
@@ -748,10 +748,12 @@ void FPersonaModule::AddLoopingInterpolation(TArray<TWeakObjectPtr<UAnimSequence
 	}
 }
 
-void FPersonaModule::CustomizeSlotNodeDetails(const TSharedRef<class IDetailsView>& InDetailsView, FOnInvokeTab InOnInvokeTab)
+void FPersonaModule::CustomizeBlueprintEditorDetails(const TSharedRef<class IDetailsView>& InDetailsView, FOnInvokeTab InOnInvokeTab)
 {
 	InDetailsView->RegisterInstancedCustomPropertyLayout(UAnimGraphNode_Slot::StaticClass(),
 		FOnGetDetailCustomizationInstance::CreateStatic(&FAnimGraphNodeSlotDetails::MakeInstance, InOnInvokeTab));
+
+	InDetailsView->SetExtensionHandler(MakeShared<FAnimGraphNodeShowAsPinExtension>());
 }
 
 IPersonaEditorModeManager* FPersonaModule::CreatePersonaEditorModeManager()
diff --git a/Engine/Source/Editor/Persona/Private/PreviewSceneCustomizations.cpp b/Engine/Source/Editor/Persona/Private/PreviewSceneCustomizations.cpp
--- a/Engine/Source/Editor/Persona/Private/PreviewSceneCustomizations.cpp
+++ b/Engine/Source/Editor/Persona/Private/PreviewSceneCustomizations.cpp
@@ -298,6 +298,38 @@ void FPreviewSceneDescriptionCustomization::CustomizeDetails(IDetailLayoutBuilde
 			LinkedAnimGraphTagProperty->MarkHiddenByCustomization();
 		}
 
+#if CHAOS_SIMULATION_DETAIL_VIEW_FACTORY_SELECTOR
+		// Physics settings
+		ClothSimulationFactoryList.Reset();
+		const TArray<IClothingSimulationFactoryClassProvider*> ClassProviders = IModularFeatures::Get().GetModularFeatureImplementations<IClothingSimulationFactoryClassProvider>(IClothingSimulationFactoryClassProvider::FeatureName);
+		for (const auto& ClassProvider : ClassProviders)
+		{
+			// Populate cloth factory list
+			ClothSimulationFactoryList.Add(MakeShared<TSubclassOf<class UClothingSimulationFactory>>(ClassProvider->GetDefaultSimulationFactoryClass()));
+		}
+
+		DetailBuilder.EditCategory("Physics")
+		.AddCustomRow(LOCTEXT("PhysicsClothingSimulationFactory", "Clothing Simulation Factory Option"))
+		.NameContent()
+		[
+			SNew(STextBlock)
+			.Font(IDetailLayoutBuilder::GetDetailFont())
+			.Text(LOCTEXT("PhysicsClothingSimulationFactory_Text", "Clothing Simulation Factory"))
+			.ToolTipText(LOCTEXT("PhysicsClothingSimulationFactory_ToolTip", "Select the cloth simulation used to preview the scene."))
+		]
+		.ValueContent()
+		.MinDesiredWidth(200.0f)
+		[
+			SNew(SComboBox<TSharedPtr<TSubclassOf<class UClothingSimulationFactory>>>)
+			.OptionsSource(&ClothSimulationFactoryList)
+			.OnGenerateWidget(this, &FPreviewSceneDescriptionCustomization::MakeClothingSimulationFactoryWidget)
+			.OnSelectionChanged(this, &FPreviewSceneDescriptionCustomization::OnClothingSimulationFactorySelectionChanged)
+			[
+				SNew(STextBlock)
+				.Text(this, &FPreviewSceneDescriptionCustomization::GetCurrentClothingSimulationFactoryText)
+			]
+		];
+#endif  // #if CHAOS_SIMULATION_DETAIL_VIEW_FACTORY_SELECTOR
 		// set the skeleton to use in our factory as we shouldn't be picking one here
 		FactoryToUse->CurrentSkeleton = EditableSkeleton.IsValid() ? MakeWeakObjectPtr(const_cast<USkeleton*>(&EditableSkeleton.Pin()->GetSkeleton())) : nullptr;
 		TArray<UFactory*> FactoriesToUse({ FactoryToUse });
@@ -671,4 +703,39 @@ void FPreviewMeshCollectionEntryCustomization::HandleMeshesArrayChanged(TSharedP
 	}
 }
 
+#if CHAOS_SIMULATION_DETAIL_VIEW_FACTORY_SELECTOR
+TSharedRef<SWidget> FPreviewSceneDescriptionCustomization::MakeClothingSimulationFactoryWidget(TSharedPtr<TSubclassOf<class UClothingSimulationFactory>> Item) const
+{
+	return SNew(STextBlock)
+		.Text(*Item ? FText::FromName((*Item)->GetFName()) : LOCTEXT("PhysicsClothingSimulationFactory_NoneSelected", "None"))
+		.Font(IDetailLayoutBuilder::GetDetailFont());
+}
+
+void FPreviewSceneDescriptionCustomization::OnClothingSimulationFactorySelectionChanged(TSharedPtr<TSubclassOf<class UClothingSimulationFactory>> Item, ESelectInfo::Type SelectInfo) const
+{
+	// Set new factory to the preview mesh component:
+	if (const TSharedPtr<IPersonaToolkit> PersonaToolkitPin = PersonaToolkit.Pin())
+	{
+		if (UDebugSkelMeshComponent* const DebugSkelMeshComponent = PersonaToolkitPin->GetPreviewMeshComponent())
+		{
+			DebugSkelMeshComponent->UnregisterComponent();
+			DebugSkelMeshComponent->ClothingSimulationFactory = *Item;
+			DebugSkelMeshComponent->RegisterComponent();
+		}
+	}
+}
+
+FText FPreviewSceneDescriptionCustomization::GetCurrentClothingSimulationFactoryText() const
+{
+	TSubclassOf<class UClothingSimulationFactory> Item;
+	if (const TSharedPtr<IPersonaToolkit> PersonaToolkitPin = PersonaToolkit.Pin())
+	{
+		if (const UDebugSkelMeshComponent* const DebugSkelMeshComponent = PersonaToolkitPin->GetPreviewMeshComponent())
+		{
+			Item = DebugSkelMeshComponent->ClothingSimulationFactory;
+		}
+	}
+	return *Item ? FText::FromName((*Item)->GetFName()) : LOCTEXT("PhysicsClothingSimulationFactory_NoneSelected", "None");
+}
+#endif  // #if CHAOS_SIMULATION_DETAIL_VIEW_FACTORY_SELECTOR
 #undef LOCTEXT_NAMESPACE
diff --git a/Engine/Source/Editor/Persona/Private/PreviewSceneCustomizations.h b/Engine/Source/Editor/Persona/Private/PreviewSceneCustomizations.h
--- a/Engine/Source/Editor/Persona/Private/PreviewSceneCustomizations.h
+++ b/Engine/Source/Editor/Persona/Private/PreviewSceneCustomizations.h
@@ -11,6 +11,10 @@
 #include "IPersonaToolkit.h"
 #include "AnimationEditorPreviewScene.h"
 
+#ifndef CHAOS_SIMULATION_DETAIL_VIEW_FACTORY_SELECTOR  // TODO: Decide whether to keep the detail view cloth selector after nvcloth has been removed
+#define CHAOS_SIMULATION_DETAIL_VIEW_FACTORY_SELECTOR WITH_CHAOS
+#endif
+
 struct FAssetData;
 class FDetailWidgetRow;
 class IDetailChildrenBuilder;
@@ -77,6 +81,17 @@ private:
 
 	ECheckBoxState HandleUseCustomAnimBPIsChecked() const;
 
+#if CHAOS_SIMULATION_DETAIL_VIEW_FACTORY_SELECTOR
+	// Make the widget of each item in the preview cloth factory combo box
+	TSharedRef<SWidget> MakeClothingSimulationFactoryWidget(TSharedPtr<TSubclassOf<class UClothingSimulationFactory>> Item) const;
+
+	// Called when the cloth factory combo box selection changes, when a new parameter type is selected
+	void OnClothingSimulationFactorySelectionChanged(TSharedPtr<TSubclassOf<class UClothingSimulationFactory>> Item, ESelectInfo::Type SelectInfo) const;
+
+	// Delegate for getting the current preview cloth factory class name as text
+	FText GetCurrentClothingSimulationFactoryText() const;
+#endif  // #if CHAOS_SIMULATION_DETAIL_VIEW_FACTORY_SELECTOR
+
 private:
 	/** Cached skeleton name to check for asset registry tags */
 	FString SkeletonName;
@@ -99,6 +114,11 @@ private:
 	/** This is list of class available to filter asset by. This list doesn't change once loaded, so only collect once */
 	static TArray<FName> AvailableClassNameList;
 
+#if CHAOS_SIMULATION_DETAIL_VIEW_FACTORY_SELECTOR
+	/** List of available cloth simulation factories. */
+	TArray<TSharedPtr<TSubclassOf<class UClothingSimulationFactory>>> ClothSimulationFactoryList;
+#endif  // #if CHAOS_SIMULATION_DETAIL_VIEW_FACTORY_SELECTOR
+
 	// Our layout builder (cached so we can refresh)
 	IDetailLayoutBuilder* MyDetailLayout;
 
diff --git a/Engine/Source/Editor/Persona/Private/SAnimCurvePicker.cpp b/Engine/Source/Editor/Persona/Private/SAnimCurvePicker.cpp
new file mode 100644
--- /dev/null
+++ b/Engine/Source/Editor/Persona/Private/SAnimCurvePicker.cpp
@@ -0,0 +1,125 @@
+// Copyright 1998-2019 Epic Games, Inc. All Rights Reserved.
+
+#include "SAnimCurvePicker.h"
+#include "AssetRegistryModule.h"
+#include "IEditableSkeleton.h"
+#include "Animation/AnimationAsset.h"
+#include "Widgets/Views/SListView.h"
+#include "Widgets/Input/SSearchBox.h"
+#include "Widgets/SBoxPanel.h"
+#include "Widgets/Text/STextBlock.h"
+#include "Widgets/Layout/SBox.h"
+
+#define LOCTEXT_NAMESPACE "SAnimCurvePicker"
+
+void SAnimCurvePicker::Construct(const FArguments& InArgs, const TSharedRef<IEditableSkeleton>& InEditableSkeleton)
+{
+	OnCurveNamePicked = InArgs._OnCurveNamePicked;
+	EditableSkeleton = InEditableSkeleton;
+
+	ChildSlot
+	[
+		SNew(SVerticalBox)
+		+SVerticalBox::Slot()
+		.AutoHeight()
+		[
+			SNew(SSearchBox)
+			.HintText(LOCTEXT("SearchBoxHint", "Search Available Curves"))
+			.OnTextChanged(this, &SAnimCurvePicker::HandleFilterTextChanged)
+		]
+		+SVerticalBox::Slot()
+		.FillHeight(1.0f)
+		[
+			SAssignNew(NameListView, SListView<TSharedPtr<FString>>)
+			.SelectionMode(ESelectionMode::Single)
+			.ItemHeight(20.0f)
+			.ListItemsSource(&CurveNames)
+			.OnSelectionChanged(this, &SAnimCurvePicker::HandleSelectionChanged)
+			.OnGenerateRow(this, &SAnimCurvePicker::HandleGenerateRow)
+		]
+	];
+
+	RefreshListItems();
+}
+
+void SAnimCurvePicker::HandleSelectionChanged(TSharedPtr<FString> InItem, ESelectInfo::Type InSelectionType)
+{
+	OnCurveNamePicked.ExecuteIfBound(FName(**InItem.Get()));
+}
+
+TSharedRef<ITableRow> SAnimCurvePicker::HandleGenerateRow(TSharedPtr<FString> InItem, const TSharedRef<STableViewBase>& InOwnerTable)
+{
+	return 
+		SNew(STableRow<TSharedPtr<FString>>, InOwnerTable)
+		[
+			SNew(SBox)
+			.MinDesiredHeight(20.0f)
+			.VAlign(VAlign_Center)
+			[
+				SNew(STextBlock)
+				.Text(FText::FromString(*InItem.Get()))
+				.HighlightText_Lambda([this]() { return FText::FromString(FilterText); })
+			]
+		];
+}
+
+void SAnimCurvePicker::RefreshListItems()
+{
+	CurveNames.Empty();
+	UniqueCurveNames.Empty();
+
+	// We use the asset registry to query all assets with the supplied skeleton, and accumulate their curve names
+	FAssetRegistryModule& AssetRegistryModule = FModuleManager::LoadModuleChecked<FAssetRegistryModule>("AssetRegistry");
+
+	FARFilter Filter;
+	Filter.bRecursiveClasses = true;
+	Filter.ClassNames.Add(UAnimationAsset::StaticClass()->GetFName());
+	Filter.TagsAndValues.Add(TEXT("Skeleton"), FAssetData(&EditableSkeleton.Pin()->GetSkeleton()).GetExportTextName());
+
+	TArray<FAssetData> FoundAssetData;
+	AssetRegistryModule.Get().GetAssets(Filter, FoundAssetData);
+
+	// now build set of unique curve names
+	for (FAssetData& AssetData : FoundAssetData)
+	{
+		const FString TagValue = AssetData.GetTagValueRef<FString>(USkeleton::CurveNameTag);
+		if (!TagValue.IsEmpty())
+		{
+			// parse notifies
+			TArray<FString> CurveValues;
+			if (TagValue.ParseIntoArray(CurveValues, *USkeleton::CurveTagDelimiter, true) > 0)
+			{
+				for (const FString& CurveValue : CurveValues)
+				{
+					UniqueCurveNames.Add(CurveValue);
+				}
+			}
+		}
+	}
+
+	FilterAvailableCurves();
+}
+
+void SAnimCurvePicker::FilterAvailableCurves()
+{
+	CurveNames.Empty();
+
+	for (const FString& UniqueCurveName : UniqueCurveNames)
+	{
+		if (FilterText.IsEmpty() || UniqueCurveName.Contains(FilterText))
+		{
+			CurveNames.Add(MakeShared<FString>(UniqueCurveName));
+		}
+	}
+
+	NameListView->RequestListRefresh();
+}
+
+void SAnimCurvePicker::HandleFilterTextChanged(const FText& InFilterText)
+{
+	FilterText = InFilterText.ToString();
+
+	FilterAvailableCurves();
+}
+
+#undef LOCTEXT_NAMESPACE
\ No newline at end of file
diff --git a/Engine/Source/Editor/Persona/Private/SAnimCurvePicker.h b/Engine/Source/Editor/Persona/Private/SAnimCurvePicker.h
new file mode 100644
--- /dev/null
+++ b/Engine/Source/Editor/Persona/Private/SAnimCurvePicker.h
@@ -0,0 +1,55 @@
+// Copyright 1998-2019 Epic Games, Inc. All Rights Reserved.
+
+#pragma once
+
+#include "CoreMinimal.h"
+#include "Widgets/DeclarativeSyntaxSupport.h"
+#include "Widgets/SCompoundWidget.h"
+#include "Widgets/Views/STableRow.h"
+
+class IEditableSkeleton;
+
+DECLARE_DELEGATE_OneParam(FOnCurveNamePicked, const FName& /*PickedName*/)
+
+class SAnimCurvePicker : public SCompoundWidget
+{
+public:
+	SLATE_BEGIN_ARGS(SAnimCurvePicker) {}
+
+	SLATE_EVENT(FOnCurveNamePicked, OnCurveNamePicked)
+
+	SLATE_END_ARGS()
+
+	void Construct(const FArguments& InArgs, const TSharedRef<IEditableSkeleton>& InEditableSkeleton);
+
+private:
+	/** Refresh the list of available curves */
+	void RefreshListItems();
+
+	/** Filter available curves */
+	void FilterAvailableCurves();
+
+	/** UI handlers */
+	void HandleSelectionChanged(TSharedPtr<FString> InItem, ESelectInfo::Type InSelectionType);
+	TSharedRef<ITableRow> HandleGenerateRow(TSharedPtr<FString> InItem, const TSharedRef<STableViewBase>& InOwnerTable);
+	void HandleFilterTextChanged(const FText& InFilterText);
+
+private:
+	/** Delegate fired when a curve name is picked */
+	FOnCurveNamePicked OnCurveNamePicked;
+
+	/** The editable skeleton we use to grab curves from */
+	TWeakPtr<IEditableSkeleton> EditableSkeleton;
+
+	/** The names of the curves we are displaying */
+	TArray<TSharedPtr<FString>> CurveNames;
+
+	/** All the unique curve names we can find */
+	TSet<FString> UniqueCurveNames;
+
+	/** The string we use to filter curve names */
+	FString FilterText;
+
+	/** The list view used to display names */
+	TSharedPtr<SListView<TSharedPtr<FString>>> NameListView;
+};
\ No newline at end of file
diff --git a/Engine/Source/Editor/Persona/Private/SAnimNotifyPanel.cpp b/Engine/Source/Editor/Persona/Private/SAnimNotifyPanel.cpp
--- a/Engine/Source/Editor/Persona/Private/SAnimNotifyPanel.cpp
+++ b/Engine/Source/Editor/Persona/Private/SAnimNotifyPanel.cpp
@@ -44,6 +44,8 @@
 #include "ClassViewerModule.h"
 #include "SSkeletonAnimNotifies.h"
 #include "Subsystems/AssetEditorSubsystem.h"
+#include "IAnimationEditor.h"
+#include "IAnimationSequenceBrowser.h"
 
 // Track Panel drawing
 const float NotificationTrackHeight = 20.0f;
@@ -823,6 +825,9 @@ protected:
 	/** Opens the supplied blueprint in an editor */
 	void OnOpenNotifySource(UBlueprint* InSourceBlueprint) const;
 
+	/** Filters the asset browser by the selected notify */
+	void OnFilterSkeletonNotify(FName InName);
+
 	/**
 	 * Selects a node on the track. Supports multi selection
 	 * @param TrackNodeIndex - Index of the node to select.
@@ -3098,19 +3103,29 @@ TSharedPtr<SWidget> SAnimNotifyTrack::SummonContextMenu(const FGeometry& MyGeome
 		UObject* NotifyObject = NotifyEvent->Notify;
 		NotifyObject = NotifyObject ? NotifyObject : NotifyEvent->NotifyStateClass;
 
-		if (NotifyObject && Cast<UBlueprintGeneratedClass>(NotifyObject->GetClass()))
+		MenuBuilder.BeginSection("ViewSource", LOCTEXT("NotifyViewHeading", "View"));
+
+		if (NotifyObject)
 		{
-			if (UBlueprint* Blueprint = Cast<UBlueprint>(NotifyObject->GetClass()->ClassGeneratedBy))
+			if (Cast<UBlueprintGeneratedClass>(NotifyObject->GetClass()))
 			{
-				MenuBuilder.BeginSection("ViewSource", LOCTEXT("NotifyViewHeading", "View"));
-
-				NewAction.ExecuteAction.BindRaw(
-					this, &SAnimNotifyTrack::OnOpenNotifySource, Blueprint);
-				MenuBuilder.AddMenuEntry(LOCTEXT("OpenNotifyBlueprint", "Open Notify Blueprint"), LOCTEXT("OpenNotifyBlueprintTooltip", "Opens the source blueprint for this notify"), FSlateIcon(), NewAction);
-
-				MenuBuilder.EndSection(); //ViewSource
+				if (UBlueprint * Blueprint = Cast<UBlueprint>(NotifyObject->GetClass()->ClassGeneratedBy))
+				{
+					NewAction.ExecuteAction.BindRaw(
+						this, &SAnimNotifyTrack::OnOpenNotifySource, Blueprint);
+					MenuBuilder.AddMenuEntry(LOCTEXT("OpenNotifyBlueprint", "Open Notify Blueprint"), LOCTEXT("OpenNotifyBlueprintTooltip", "Opens the source blueprint for this notify"), FSlateIcon(), NewAction);
+				}
 			}
 		}
+		else
+		{
+			// skeleton notify
+			NewAction.ExecuteAction.BindRaw(
+				this, &SAnimNotifyTrack::OnFilterSkeletonNotify, NotifyEvent->NotifyName);
+			MenuBuilder.AddMenuEntry(LOCTEXT("FindNotifyReferences", "Find References"), LOCTEXT("FindNotifyReferencesTooltip", "Find all references to this skeleton notify in the asset browser"), FSlateIcon(), NewAction);
+		}
+
+		MenuBuilder.EndSection(); //ViewSource
 	}
 
 	FWidgetPath WidgetPath = MouseEvent.GetEventPath() != nullptr ? *MouseEvent.GetEventPath() : FWidgetPath();
@@ -3169,6 +3184,17 @@ void SAnimNotifyTrack::OnOpenNotifySource(UBlueprint* InSourceBlueprint) const
 	GEditor->GetEditorSubsystem<UAssetEditorSubsystem>()->OpenEditorForAsset(InSourceBlueprint);
 }
 
+void SAnimNotifyTrack::OnFilterSkeletonNotify(FName InName)
+{
+	// Open asset browser first
+	OnInvokeTab.ExecuteIfBound(FPersonaTabs::AssetBrowserID);
+
+	IAssetEditorInstance* AssetEditor = GEditor->GetEditorSubsystem<UAssetEditorSubsystem>()->FindEditorForAsset(Sequence, true);
+	check(AssetEditor->GetEditorName() == TEXT("AnimationEditor"));
+	IAnimationEditor* AnimationEditor = static_cast<IAnimationEditor*>(AssetEditor);
+	AnimationEditor->GetAssetBrowser()->FilterBySkeletonNotify(InName);
+}
+
 bool SAnimNotifyTrack::IsSingleNodeSelected()
 {
 	return SelectedNodeIndices.Num() == 1;
diff --git a/Engine/Source/Editor/Persona/Private/SAnimViewportToolBar.cpp b/Engine/Source/Editor/Persona/Private/SAnimViewportToolBar.cpp
--- a/Engine/Source/Editor/Persona/Private/SAnimViewportToolBar.cpp
+++ b/Engine/Source/Editor/Persona/Private/SAnimViewportToolBar.cpp
@@ -682,7 +682,7 @@ TSharedRef<SWidget> SAnimViewportToolBar::GenerateCharacterMenu() const
 				})
 				);
 
-#if WITH_APEX_CLOTHING
+#if WITH_APEX_CLOTHING || WITH_CHAOS_CLOTHING
 			UDebugSkelMeshComponent* PreviewComp = Viewport.Pin()->GetPreviewScene()->GetPreviewMeshComponent();
 
 			if(PreviewComp)
@@ -692,7 +692,7 @@ TSharedRef<SWidget> SAnimViewportToolBar::GenerateCharacterMenu() const
 					LOCTEXT("CharacterMenu_ClothingSubMenuToolTip", "Options relating to clothing"),
 					FNewMenuDelegate::CreateRaw(const_cast<SAnimViewportToolBar*>(this), &SAnimViewportToolBar::FillCharacterClothingMenu));
 			}
-#endif // #if WITH_APEX_CLOTHING
+#endif // #if WITH_APEX_CLOTHING || WITH_CHAOS_CLOTHING
 		}
 
 		InMenuBuilder.AddSubMenu(
@@ -760,7 +760,7 @@ void SAnimViewportToolBar::FillCharacterAdvancedMenu(FMenuBuilder& MenuBuilder)
 
 void SAnimViewportToolBar::FillCharacterClothingMenu(FMenuBuilder& MenuBuilder)
 {
-#if WITH_APEX_CLOTHING
+#if WITH_APEX_CLOTHING || WITH_CHAOS_CLOTHING
 	const FAnimViewportShowCommands& Actions = FAnimViewportShowCommands::Get();
 
 	MenuBuilder.BeginSection("ClothPreview", LOCTEXT("ClothPreview_Label", "Simulation"));
@@ -803,7 +803,7 @@ void SAnimViewportToolBar::FillCharacterClothingMenu(FMenuBuilder& MenuBuilder)
 		}
 	}
 
-#endif // #if WITH_APEX_CLOTHING
+#endif // #if WITH_APEX_CLOTHING || WITH_CHAOS_CLOTHING
 }
 
 TSharedRef<SWidget> SAnimViewportToolBar::GenerateShowMenu() const
diff --git a/Engine/Source/Editor/Persona/Private/SAnimationBlendSpaceGridWidget.cpp b/Engine/Source/Editor/Persona/Private/SAnimationBlendSpaceGridWidget.cpp
--- a/Engine/Source/Editor/Persona/Private/SAnimationBlendSpaceGridWidget.cpp
+++ b/Engine/Source/Editor/Persona/Private/SAnimationBlendSpaceGridWidget.cpp
@@ -82,6 +82,10 @@ void SBlendSpaceGridWidget::Construct(const FArguments& InArgs)
 	// Retrieve background and sample key brushes 
 	BackgroundImage = FEditorStyle::GetBrush(TEXT("Graph.Panel.SolidBackground"));
 	KeyBrush = FEditorStyle::GetBrush("CurveEd.CurveKey");
+	ArrowBrushes[(uint8)EArrowDirection::Up] = FEditorStyle::GetBrush("BlendSpaceEditor.ArrowUp");
+	ArrowBrushes[(uint8)EArrowDirection::Down] = FEditorStyle::GetBrush("BlendSpaceEditor.ArrowDown");
+	ArrowBrushes[(uint8)EArrowDirection::Right] = FEditorStyle::GetBrush("BlendSpaceEditor.ArrowRight");
+	ArrowBrushes[(uint8)EArrowDirection::Left] = FEditorStyle::GetBrush("BlendSpaceEditor.ArrowLeft");
 
 	// Retrieve font data 
 	FontInfo = FEditorStyle::GetFontStyle("CurveEd.InfoFont");
@@ -412,8 +416,19 @@ void SBlendSpaceGridWidget::PaintAxisText(const FGeometry& AllottedGeometry, con
 	FString Text = ParameterXName.ToString();
 	FVector2D TextSize = FontMeasure->Measure(Text, FontInfo);
 
+	// arrow left
+	FVector2D TextPosition = FVector2D(GridCenter.X - (TextSize.X * .5f), CachedGridRectangle.Bottom + TextMargin + (KeySize.Y * .25f));
+	FVector2D ArrowSize = ArrowBrushes[(uint8)EArrowDirection::Left]->GetImageSize();
+	FVector2D ArrowPosition = FVector2D(TextPosition.X - ArrowSize.X - 10.f/* give padding*/, TextPosition.Y);
+	FSlateDrawElement::MakeBox(OutDrawElements, DrawLayerId + 1, AllottedGeometry.ToPaintGeometry(ArrowPosition, ArrowSize), ArrowBrushes[(uint8)EArrowDirection::Left], ESlateDrawEffect::None, FLinearColor::White);
+
 	// Label
-	FSlateDrawElement::MakeText(OutDrawElements, DrawLayerId + 1, AllottedGeometry.MakeChild(FVector2D(GridCenter.X - (TextSize.X * .5f), CachedGridRectangle.Bottom + TextMargin + (KeySize.Y * .25f)), FVector2D(1.0f, 1.0f)).ToPaintGeometry(), Text, FontInfo, ESlateDrawEffect::None, FLinearColor::White);
+	FSlateDrawElement::MakeText(OutDrawElements, DrawLayerId + 1, AllottedGeometry.MakeChild(TextPosition, FVector2D(1.0f, 1.0f)).ToPaintGeometry(), Text, FontInfo, ESlateDrawEffect::None, FLinearColor::White);
+
+	// arrow right
+	ArrowSize = ArrowBrushes[(uint8)EArrowDirection::Right]->GetImageSize();
+	ArrowPosition = FVector2D(TextPosition.X + TextSize.X + 10.f/* give padding*/, TextPosition.Y);
+	FSlateDrawElement::MakeBox(OutDrawElements, DrawLayerId + 1, AllottedGeometry.ToPaintGeometry(ArrowPosition, ArrowSize), ArrowBrushes[(uint8)EArrowDirection::Right], ESlateDrawEffect::None, FLinearColor::White);
 
 	Text = FString::SanitizeFloat(SampleValueMin.X);
 	TextSize = FontMeasure->Measure(Text, FontInfo);
@@ -434,8 +449,19 @@ void SBlendSpaceGridWidget::PaintAxisText(const FGeometry& AllottedGeometry, con
 		Text = ParameterYName.ToString();
 		TextSize = FontMeasure->Measure(Text, FontInfo);
 
+		// arrow up
+		TextPosition = FVector2D(((GridMargin.Left - TextSize.X) * 0.5f - (KeySize.X * .25f)) + GridRatioMargin.Left, GridCenter.Y - (TextSize.Y * .5f));
+		ArrowSize = ArrowBrushes[(uint8)EArrowDirection::Up]->GetImageSize();
+		ArrowPosition = FVector2D(TextPosition.X + TextSize.X * 0.5f - ArrowSize.X * 0.5f, TextPosition.Y - ArrowSize.X - 10.f/* give padding*/);
+		FSlateDrawElement::MakeBox(OutDrawElements, DrawLayerId + 1, AllottedGeometry.ToPaintGeometry(ArrowPosition, ArrowSize), ArrowBrushes[(uint8)EArrowDirection::Up], ESlateDrawEffect::None, FLinearColor::White);
+
 		// Label
-		FSlateDrawElement::MakeText(OutDrawElements, DrawLayerId + 1, AllottedGeometry.MakeChild(FVector2D(((GridMargin.Left - TextSize.X) * 0.5f - (KeySize.X * .25f)) + GridRatioMargin.Left, GridCenter.Y - (TextSize.Y * .5f)), FVector2D(1.0f, 1.0f)).ToPaintGeometry(), Text,	FontInfo, ESlateDrawEffect::None, FLinearColor::White);
+		FSlateDrawElement::MakeText(OutDrawElements, DrawLayerId + 1, AllottedGeometry.MakeChild(TextPosition, FVector2D(1.0f, 1.0f)).ToPaintGeometry(), Text,	FontInfo, ESlateDrawEffect::None, FLinearColor::White);
+
+		// arrow down
+		ArrowSize = ArrowBrushes[(uint8)EArrowDirection::Down]->GetImageSize();
+		ArrowPosition = FVector2D(TextPosition.X + TextSize.X * 0.5f - ArrowSize.X * 0.5f, TextPosition.Y + TextSize.Y + 10.f/* give padding*/);
+		FSlateDrawElement::MakeBox(OutDrawElements, DrawLayerId + 1, AllottedGeometry.ToPaintGeometry(ArrowPosition, ArrowSize), ArrowBrushes[(uint8)EArrowDirection::Down], ESlateDrawEffect::None, FLinearColor::White);
 
 		Text = FString::SanitizeFloat(SampleValueMin.Y);
 		TextSize = FontMeasure->Measure(Text, FontInfo);
diff --git a/Engine/Source/Editor/Persona/Private/SAnimationBlendSpaceGridWidget.h b/Engine/Source/Editor/Persona/Private/SAnimationBlendSpaceGridWidget.h
--- a/Engine/Source/Editor/Persona/Private/SAnimationBlendSpaceGridWidget.h
+++ b/Engine/Source/Editor/Persona/Private/SAnimationBlendSpaceGridWidget.h
@@ -269,6 +269,15 @@ private:
 	FVector2D KeySize;
 	const FSlateBrush* KeyBrush;
 	const FSlateBrush* BackgroundImage;
+	enum class EArrowDirection : uint8 
+	{
+		Left, 
+		Right, 
+		Up, 
+		Down, 
+		Max
+	};
+	const FSlateBrush* ArrowBrushes[(int32)EArrowDirection::Max];
 
 	/** Individual sample state colours */
 	FSlateColor HighlightKeyColor;
diff --git a/Engine/Source/Editor/Persona/Private/SAnimationEditorViewport.h b/Engine/Source/Editor/Persona/Private/SAnimationEditorViewport.h
--- a/Engine/Source/Editor/Persona/Private/SAnimationEditorViewport.h
+++ b/Engine/Source/Editor/Persona/Private/SAnimationEditorViewport.h
@@ -497,7 +497,7 @@ public:
 
 	bool IsTurnTableSpeedSelected(int32 SpeedIndex) const;
 
-#if WITH_APEX_CLOTHING
+#if WITH_APEX_CLOTHING || WITH_CHAOS_CLOTHING
 	/** 
 	 * clothing show options 
 	*/
@@ -520,7 +520,7 @@ private:
 	void OnSetSectionsDisplayMode(ESectionDisplayMode DisplayMode);
 	bool IsSectionsDisplayMode(ESectionDisplayMode DisplayMode) const;
 
-#endif // #if WITH_APEX_CLOTHING
+#endif // #if WITH_APEX_CLOTHING || WITH_CHAOS_CLOTHING
 
 private:
 	/** Weak pointer back to the preview scene we are viewing */
diff --git a/Engine/Source/Editor/Persona/Private/SAnimationSequenceBrowser.cpp b/Engine/Source/Editor/Persona/Private/SAnimationSequenceBrowser.cpp
--- a/Engine/Source/Editor/Persona/Private/SAnimationSequenceBrowser.cpp
+++ b/Engine/Source/Editor/Persona/Private/SAnimationSequenceBrowser.cpp
@@ -41,6 +41,10 @@
 #include "Preferences/PersonaOptions.h"
 #include "Settings/SkeletalMeshEditorSettings.h"
 #include "Subsystems/AssetEditorSubsystem.h"
+#include "SSkeletonAnimNotifies.h"
+#include "SAnimCurveViewer.h"
+#include "IEditableSkeleton.h"
+#include "SAnimCurvePicker.h"
 
 #define LOCTEXT_NAMESPACE "SequenceBrowser"
 
@@ -65,6 +69,199 @@ public:
 	}
 };
 
+/** A filter that displays animations that use a skeleton notify */
+class FFrontendFilter_SkeletonNotify : public FFrontendFilter
+{
+public:
+	FFrontendFilter_SkeletonNotify(TSharedPtr<FFrontendFilterCategory> InCategory, TSharedPtr<IEditableSkeleton> InEditableSkeleton)
+		: FFrontendFilter(InCategory) 
+		, EditableSkeleton(InEditableSkeleton)
+	{}
+
+	// FFrontendFilter implementation
+	virtual FString GetName() const override 
+	{
+		return TEXT("SkeletonNotifyAssets"); 
+	}
+
+	virtual FText GetDisplayName() const override 
+	{ 
+		return NotifyString.IsEmpty() ? LOCTEXT("FFrontendFilter_SkeletonNotifyAssetsEmpty", "Uses Skeleton Notify...") : FText::Format(LOCTEXT("FFrontendFilter_SkeletonNotifyAssets", "Uses Skeleton Notify: {0}"), FText::FromString(NotifyString));
+	}
+
+	virtual FText GetToolTipText() const override
+	{
+		return NotifyString.IsEmpty() ? LOCTEXT("FFrontendFilter_SkeletonNotifyAssetsEmptyTooltip", "Show assets that use a specified skeleton notify") : FText::Format(LOCTEXT("FFrontendFilter_SkeletonNotifyAssetsTooltip", "Show assets that use skeleton notify '{0}'"), FText::FromString(NotifyString));
+	}
+
+	virtual void ModifyContextMenu(FMenuBuilder& MenuBuilder) override
+	{
+		MenuBuilder.BeginSection(TEXT("NotifySection"), LOCTEXT("NotifySectionHeading", "Choose Notify"));
+
+		TSharedRef<SWidget> Widget =
+			SNew(SBox)
+			.HeightOverride(300.0f)
+			.WidthOverride(200.0f)
+			[
+				SNew(SSkeletonAnimNotifies, EditableSkeleton.ToSharedRef())
+				.IsPicker(true)
+				.OnItemSelected_Lambda([this](const FName& InNotifyName)
+				{
+					FSlateApplication::Get().DismissAllMenus();
+					NotifyString = InNotifyName.ToString();
+					OnChanged().Broadcast();
+				})
+			];
+
+		MenuBuilder.AddWidget(Widget, FText(), true);
+
+		MenuBuilder.EndSection();
+	}
+
+	virtual FLinearColor GetColor() const override
+	{
+		return FLinearColor(0.0f, 0.4f, 0.6f, 1);
+	}
+
+	// IFilter implementation
+	virtual bool PassesFilter(FAssetFilterType InItem) const override
+	{
+		if (!NotifyString.IsEmpty())
+		{
+			const FString TagValue = InItem.GetTagValueRef<FString>(USkeleton::AnimNotifyTag);
+			if (!TagValue.IsEmpty())
+			{
+				// parse notifies
+				TArray<FString> NotifyValues;
+				if (TagValue.ParseIntoArray(NotifyValues, *USkeleton::AnimNotifyTagDelimiter, true) > 0)
+				{
+					for (const FString& NotifyValue : NotifyValues)
+					{
+						if (NotifyValue == NotifyString)
+						{
+							return true;
+						}
+					}
+				}
+			}
+
+			return false;
+		}
+
+		return true;
+	}
+
+	void SetNotifyFilter(const FName& InNotifyFilter)
+	{
+		NotifyString = InNotifyFilter.ToString();
+		OnChanged().Broadcast();
+		SetActive(true);
+	}
+
+public:
+	/** Notify string to use when filtering */
+	FString NotifyString;
+
+	/** Editable skeleton used to access available notifies */
+	TSharedPtr<IEditableSkeleton> EditableSkeleton;
+};
+
+/** A filter that displays animations that use a curve */
+class FFrontendFilter_Curves : public FFrontendFilter
+{
+public:
+	FFrontendFilter_Curves(TSharedPtr<FFrontendFilterCategory> InCategory, TSharedPtr<IEditableSkeleton> InEditableSkeleton)
+		: FFrontendFilter(InCategory)
+		, EditableSkeleton(InEditableSkeleton)
+	{}
+
+	// FFrontendFilter implementation
+	virtual FString GetName() const override
+	{
+		return TEXT("CurveAssets");
+	}
+
+	virtual FText GetDisplayName() const override
+	{
+		return CurveString.IsEmpty() ? LOCTEXT("FFrontendFilter_CurveAssetsEmpty", "Uses Curve...") : FText::Format(LOCTEXT("FFrontendFilter_CurveAssets", "Uses Curve: {0}"), FText::FromString(CurveString));
+	}
+
+	virtual FText GetToolTipText() const override
+	{
+		return CurveString.IsEmpty() ? LOCTEXT("FFrontendFilter_CurveAssetsEmptyTooltip", "Show assets that use a specified curve") : FText::Format(LOCTEXT("FFrontendFilter_CurveAssetsTooltip", "Show assets that use curve '{0}'"), FText::FromString(CurveString));
+	}
+
+	virtual void ModifyContextMenu(FMenuBuilder& MenuBuilder) override
+	{
+		MenuBuilder.BeginSection(TEXT("CurveSection"), LOCTEXT("CurveSectionHeading", "Choose Curve"));
+
+		TSharedRef<SWidget> Widget = 
+			SNew(SBox)
+			.HeightOverride(300.0f)
+			.WidthOverride(200.0f)
+			[
+				SNew(SAnimCurvePicker, EditableSkeleton.ToSharedRef())
+				.OnCurveNamePicked_Lambda([this](const FName& InCurveName)
+				{
+					FSlateApplication::Get().DismissAllMenus();
+					CurveString = InCurveName.ToString();
+					OnChanged().Broadcast();
+				})
+			];
+
+		MenuBuilder.AddWidget(Widget, FText(), true);
+
+		MenuBuilder.EndSection();
+	}
+
+	virtual FLinearColor GetColor() const override
+	{
+		return FLinearColor(0.0f, 0.6f, 0.2f, 1);
+	}
+
+	// IFilter implementation
+	virtual bool PassesFilter(FAssetFilterType InItem) const override
+	{
+		if (!CurveString.IsEmpty())
+		{
+			const FString TagValue = InItem.GetTagValueRef<FString>(USkeleton::CurveNameTag);
+			if (!TagValue.IsEmpty())
+			{
+				// parse curves
+				TArray<FString> CurveValues;
+				if (TagValue.ParseIntoArray(CurveValues, *USkeleton::CurveTagDelimiter, true) > 0)
+				{
+					for (const FString& CurveValue : CurveValues)
+					{
+						if (CurveValue == CurveString)
+						{
+							return true;
+						}
+					}
+				}
+			}
+
+			return false;
+		}
+
+		return true;
+	}
+
+	void SetCurveFilter(const FName& InCurveFilter)
+	{
+		CurveString = InCurveFilter.ToString();
+		OnChanged().Broadcast();
+		SetActive(true);
+	}
+
+public:
+	/** Curve string to use when filtering */
+	FString CurveString;
+
+	/** Editable skeleton used to access available curves */
+	TSharedPtr<IEditableSkeleton> EditableSkeleton;
+};
+
 /** A filter that displays sound waves */
 class FFrontendFilter_SoundWaves : public FFrontendFilter
 {
@@ -635,6 +832,10 @@ void SAnimationSequenceBrowser::Construct(const FArguments& InArgs, const TShare
 		Config.ExtraFrontendFilters.Add(FolderFilter);
 	}
 
+	SkeletonNotifyFilter = MakeShareable(new FFrontendFilter_SkeletonNotify(AnimCategory, InPersonaToolkit->GetEditableSkeleton()));
+	Config.ExtraFrontendFilters.Add(SkeletonNotifyFilter.ToSharedRef());
+	Config.ExtraFrontendFilters.Add(MakeShareable(new FFrontendFilter_Curves(AnimCategory, InPersonaToolkit->GetEditableSkeleton())));
+
 	Config.OnIsAssetValidForCustomToolTip = FOnIsAssetValidForCustomToolTip::CreateLambda([](const FAssetData& AssetData) {return AssetData.IsAssetLoaded(); });
 	Config.OnGetCustomAssetToolTip = FOnGetCustomAssetToolTip::CreateSP(this, &SAnimationSequenceBrowser::CreateCustomAssetToolTip);
 	Config.OnVisualizeAssetToolTip = FOnVisualizeAssetToolTip::CreateSP(this, &SAnimationSequenceBrowser::OnVisualizeAssetToolTip);
@@ -956,6 +1157,13 @@ void SAnimationSequenceBrowser::SelectAsset(UAnimationAsset * AnimAsset)
 	}
 }
 
+void SAnimationSequenceBrowser::FilterBySkeletonNotify(const FName& InNotifyName)
+{
+	check(SkeletonNotifyFilter.IsValid());
+
+	SkeletonNotifyFilter->SetNotifyFilter(InNotifyName);
+}
+
 void SAnimationSequenceBrowser::AddToHistory(UAnimationAsset * AnimAsset)
 {
 	if (AnimAsset)
diff --git a/Engine/Source/Editor/Persona/Private/SAnimationSequenceBrowser.h b/Engine/Source/Editor/Persona/Private/SAnimationSequenceBrowser.h
--- a/Engine/Source/Editor/Persona/Private/SAnimationSequenceBrowser.h
+++ b/Engine/Source/Editor/Persona/Private/SAnimationSequenceBrowser.h
@@ -27,6 +27,7 @@ class UAnimationAsset;
 class UDebugSkelMeshComponent;
 class USoundWave;
 class FFrontendFilter_Folder;
+class FFrontendFilter_SkeletonNotify;
 
 //////////////////////////////////////////////////////////////////////////
 // FAnimationAssetViewportClient
@@ -67,7 +68,9 @@ public:
 
 	/** IAnimationSequenceBrowser interface */
 	virtual void SelectAsset(UAnimationAsset * AnimAsset) override;
-
+	virtual void AddToHistory(UAnimationAsset* AnimAsset) override;
+	virtual void FilterBySkeletonNotify(const FName& InNotifyName) override;
+	
 	/** Delegate that handles creation of context menu */
 	TSharedPtr<SWidget> OnGetAssetContextMenu(const TArray<FAssetData>& SelectedAssets);
 
@@ -104,9 +107,6 @@ public:
 	 */
 	void OnCreateCopy(TArray<FAssetData> Selected);
 
-	/** public reference to add to history */
-	virtual void AddToHistory(UAnimationAsset * AnimAsset) override;
-
 protected:
 	bool CanShowColumnForAssetRegistryTag(FName AssetType, FName TagName) const;
 	
@@ -262,6 +262,9 @@ protected:
 	/** All the folder filters we have */
 	TArray<TSharedPtr<FFrontendFilter_Folder>> FolderFilters;
 
+	/** The skeleton notify filter */
+	TSharedPtr<FFrontendFilter_SkeletonNotify> SkeletonNotifyFilter;
+
 	void RetargetAnimationHandler(USkeleton* OldSkeleton, USkeleton* NewSkeleton, bool bRemapReferencedAssets, bool bAllowRemapToExisting, bool bConvertSpaces, const EditorAnimUtils::FNameDuplicationRule* NameRule, TArray<TWeakObjectPtr<UObject>> InAnimAssets);
 
 private:
diff --git a/Engine/Source/Editor/Persona/Private/SSkeletonAnimNotifies.cpp b/Engine/Source/Editor/Persona/Private/SSkeletonAnimNotifies.cpp
--- a/Engine/Source/Editor/Persona/Private/SSkeletonAnimNotifies.cpp
+++ b/Engine/Source/Editor/Persona/Private/SSkeletonAnimNotifies.cpp
@@ -16,6 +16,9 @@
 #include "IEditableSkeleton.h"
 #include "TabSpawners.h"
 #include "Editor.h"
+#include "IAnimationEditor.h"
+#include "IAnimationBlueprintEditor.h"
+#include "IAnimationSequenceBrowser.h"
 
 #define LOCTEXT_NAMESPACE "SkeletonAnimNotifies"
 
@@ -44,14 +47,14 @@ FSkeletonAnimNotifiesSummoner::FSkeletonAnimNotifiesSummoner(TSharedPtr<class FA
 
 TSharedRef<SWidget> FSkeletonAnimNotifiesSummoner::CreateTabBody(const FWorkflowTabSpawnInfo& Info) const
 {
-	return SNew(SSkeletonAnimNotifies, EditableSkeleton.Pin().ToSharedRef())
+	return SNew(SSkeletonAnimNotifies, EditableSkeleton.Pin().ToSharedRef(), HostingApp.Pin())
 		.OnObjectsSelected(OnObjectsSelected);
 }
 
 /////////////////////////////////////////////////////
 // SSkeletonAnimNotifies
 
-void SSkeletonAnimNotifies::Construct(const FArguments& InArgs, const TSharedRef<IEditableSkeleton>& InEditableSkeleton)
+void SSkeletonAnimNotifies::Construct(const FArguments& InArgs, const TSharedRef<IEditableSkeleton>& InEditableSkeleton, const TSharedPtr<class FAssetEditorToolkit>& InHostingApp)
 {
 	OnObjectsSelected = InArgs._OnObjectsSelected;
 	OnItemSelected = InArgs._OnItemSelected;
@@ -59,6 +62,7 @@ void SSkeletonAnimNotifies::Construct(const FArguments& InArgs, const TSharedRef
 	bIsSyncMarker = InArgs._IsSyncMarker;
 
 	EditableSkeleton = InEditableSkeleton;
+	WeakHostingApp = InHostingApp;
 
 	if (bIsSyncMarker)
 	{
@@ -202,7 +206,14 @@ TSharedPtr<SWidget> SSkeletonAnimNotifies::OnGetContextMenuContent() const
 				MenuBuilder.AddMenuEntry(Label, ToolTipText, FSlateIcon(), Action);
 			}
 
-
+			if(WeakHostingApp.IsValid() && NotifiesListView->GetNumItemsSelected() == 1)
+			{
+				FUIAction Action = FUIAction(FExecuteAction::CreateSP(const_cast<SSkeletonAnimNotifies*>(this), &SSkeletonAnimNotifies::OnFindReferences),
+					FCanExecuteAction::CreateSP(this, &SSkeletonAnimNotifies::CanPerformDelete));
+				const FText Label = LOCTEXT("FindNotifyReferences", "Find References");
+				const FText ToolTipText = LOCTEXT("FindNotifyReferencesTooltip", "Find all references to this skeleton notify in the asset browser");
+				MenuBuilder.AddMenuEntry(Label, ToolTipText, FSlateIcon(), Action);
+			}
 		}
 		MenuBuilder.EndSection();
 	}
@@ -462,4 +473,30 @@ void SSkeletonAnimNotifies::NotifyUser( FNotificationInfo& NotificationInfo )
 		Notification->SetCompletionState( SNotificationItem::CS_Fail );
 	}
 }
+
+void SSkeletonAnimNotifies::OnFindReferences()
+{
+	TSharedPtr<FAssetEditorToolkit> HostingApp = WeakHostingApp.Pin();
+	if (HostingApp.IsValid())
+	{
+		check(NotifiesListView->GetNumItemsSelected() == 1);
+		TArray<TSharedPtr<FDisplayedAnimNotifyInfo>> SelectedItems;
+		NotifiesListView->GetSelectedItems(SelectedItems);
+		FName Name = SelectedItems[0]->Name;
+
+		HostingApp->InvokeTab(FPersonaTabs::AssetBrowserID);
+
+		if (HostingApp->GetEditorName() == TEXT("AnimationEditor"))
+		{
+			TSharedPtr<IAnimationEditor> AnimationEditor = StaticCastSharedPtr<IAnimationEditor>(HostingApp);
+			AnimationEditor->GetAssetBrowser()->FilterBySkeletonNotify(Name);
+		}
+		else if (HostingApp->GetEditorName() == TEXT("AnimationBlueprintEditor"))
+		{
+			TSharedPtr<IAnimationBlueprintEditor> AnimationBlueprintEditor = StaticCastSharedPtr<IAnimationBlueprintEditor>(HostingApp);
+			AnimationBlueprintEditor->GetAssetBrowser()->FilterBySkeletonNotify(Name);
+		}
+	}
+}
+
 #undef LOCTEXT_NAMESPACE
diff --git a/Engine/Source/Editor/Persona/Private/SSkeletonAnimNotifies.h b/Engine/Source/Editor/Persona/Private/SSkeletonAnimNotifies.h
--- a/Engine/Source/Editor/Persona/Private/SSkeletonAnimNotifies.h
+++ b/Engine/Source/Editor/Persona/Private/SSkeletonAnimNotifies.h
@@ -102,7 +102,7 @@ public:
 
 	SLATE_END_ARGS()
 public:
-	void Construct(const FArguments& InArgs, const TSharedRef<IEditableSkeleton>& InEditableSkeleton);
+	void Construct(const FArguments& InArgs, const TSharedRef<IEditableSkeleton>& InEditableSkeleton, const TSharedPtr<class FAssetEditorToolkit>& InHostingApp = nullptr);
 
 	~SSkeletonAnimNotifies();
 
@@ -187,6 +187,9 @@ private:
 	/** Handle when an item is scrolled into view, triggers a rename for new items */
 	void OnItemScrolledIntoView(TSharedPtr<FDisplayedAnimNotifyInfo> InItem, const TSharedPtr<ITableRow>& InTableRow);
 
+	/** Handle find references from the context menu */
+	void OnFindReferences();
+
 	/** The skeleton we are currently editing */
 	TSharedPtr<class IEditableSkeleton> EditableSkeleton;
 
@@ -211,6 +214,9 @@ private:
 	/** Delegate fired when a notify is selected */
 	FOnItemSelected OnItemSelected;
 
+	/** The app we are embedded in */
+	TWeakPtr<class FAssetEditorToolkit> WeakHostingApp;
+
 	/** Whether we should use this dialog as a picker or an editor. In picker mode we cant add, remove or rename notifies. */
 	bool bIsPicker;
 
diff --git a/Engine/Source/Editor/Persona/Public/IAnimationSequenceBrowser.h b/Engine/Source/Editor/Persona/Public/IAnimationSequenceBrowser.h
--- a/Engine/Source/Editor/Persona/Public/IAnimationSequenceBrowser.h
+++ b/Engine/Source/Editor/Persona/Public/IAnimationSequenceBrowser.h
@@ -12,6 +12,10 @@ class IAnimationSequenceBrowser : public SCompoundWidget
 public:
 	/** Select the specified asset */
 	virtual void SelectAsset(UAnimationAsset* InAsset) = 0;
+
 	/** Add to navigation history */
-	virtual void AddToHistory(UAnimationAsset * AnimAsset) = 0;
+	virtual void AddToHistory(UAnimationAsset* AnimAsset) = 0;
+
+	/** Add filter by skeleton notify name */
+	virtual void FilterBySkeletonNotify(const FName& InNotifyName) = 0;
 };
diff --git a/Engine/Source/Editor/Persona/Public/PersonaModule.h b/Engine/Source/Editor/Persona/Public/PersonaModule.h
--- a/Engine/Source/Editor/Persona/Public/PersonaModule.h
+++ b/Engine/Source/Editor/Persona/Public/PersonaModule.h
@@ -269,8 +269,11 @@ public:
 	/** Add looping interpolation to the list of animations */
 	virtual void AddLoopingInterpolation(TArray<TWeakObjectPtr<class UAnimSequence>>& AnimSequences);
 
+	UE_DEPRECATED(4.24, "Function renamed, please use CustomizeBlueprintEditorDetails")
+	virtual void CustomizeSlotNodeDetails(const TSharedRef<class IDetailsView>& InDetailsView, FOnInvokeTab InOnInvokeTab) { CustomizeBlueprintEditorDetails(InDetailsView, InOnInvokeTab); }
+
 	/** Customize the details of a slot node for the specified details view */
-	virtual void CustomizeSlotNodeDetails(const TSharedRef<class IDetailsView>& InDetailsView, FOnInvokeTab InOnInvokeTab);
+	virtual void CustomizeBlueprintEditorDetails(const TSharedRef<class IDetailsView>& InDetailsView, FOnInvokeTab InOnInvokeTab);
 
 	/** Create a Persona editor mode manager. Should be destroyed using plain ol' delete. Note: Only FPersonaEditMode-derived modes should be used with this manager! */
 	virtual class IPersonaEditorModeManager* CreatePersonaEditorModeManager();
diff --git a/Engine/Source/Editor/PropertyEditor/Private/DetailLayoutBuilderImpl.cpp b/Engine/Source/Editor/PropertyEditor/Private/DetailLayoutBuilderImpl.cpp
--- a/Engine/Source/Editor/PropertyEditor/Private/DetailLayoutBuilderImpl.cpp
+++ b/Engine/Source/Editor/PropertyEditor/Private/DetailLayoutBuilderImpl.cpp
@@ -152,7 +152,7 @@ IDetailPropertyRow* FDetailLayoutBuilderImpl::EditDefaultProperty(TSharedPtr<IPr
 	return nullptr;
 }
 
-TSharedRef<IPropertyHandle> FDetailLayoutBuilderImpl::GetProperty( const FName PropertyPath, const UStruct* ClassOutermost, FName InInstanceName )
+TSharedRef<IPropertyHandle> FDetailLayoutBuilderImpl::GetProperty( const FName PropertyPath, const UStruct* ClassOutermost, FName InInstanceName ) const
 {	
 	TSharedPtr<FPropertyHandleBase> PropertyHandle; 
 
@@ -551,7 +551,7 @@ TSharedPtr<FPropertyNode> FDetailLayoutBuilderImpl::GetPropertyNodeInternal( con
 }
 
 
-TSharedRef<IPropertyHandle> FDetailLayoutBuilderImpl::GetPropertyHandle( TSharedPtr<FPropertyNode> PropertyNodePtr )
+TSharedRef<IPropertyHandle> FDetailLayoutBuilderImpl::GetPropertyHandle( TSharedPtr<FPropertyNode> PropertyNodePtr ) const
 {
 	TSharedPtr<IPropertyHandle> PropertyHandle;
 	if( PropertyNodePtr.IsValid() )
diff --git a/Engine/Source/Editor/PropertyEditor/Private/DetailLayoutBuilderImpl.h b/Engine/Source/Editor/PropertyEditor/Private/DetailLayoutBuilderImpl.h
--- a/Engine/Source/Editor/PropertyEditor/Private/DetailLayoutBuilderImpl.h
+++ b/Engine/Source/Editor/PropertyEditor/Private/DetailLayoutBuilderImpl.h
@@ -36,7 +36,7 @@ public:
 	virtual IDetailPropertyRow& AddPropertyToCategory(TSharedPtr<IPropertyHandle> InPropertyHandle) override;
 	virtual FDetailWidgetRow& AddCustomRowToCategory(TSharedPtr<IPropertyHandle> InPropertyHandle, const FText& InCustomSearchString, bool bForAdvanced = false) override;
 	virtual IDetailPropertyRow* EditDefaultProperty(TSharedPtr<IPropertyHandle> InPropertyHandle) override;
-	virtual TSharedRef<IPropertyHandle> GetProperty(const FName PropertyPath, const UStruct* ClassOutermost, FName InInstanceName) override;
+	virtual TSharedRef<IPropertyHandle> GetProperty(const FName PropertyPath, const UStruct* ClassOutermost, FName InInstanceName) const override;
 	virtual FName GetTopLevelProperty() override;
 	virtual void HideProperty(const TSharedPtr<IPropertyHandle> Property) override;
 	virtual void HideProperty(FName PropertyPath, const UStruct* ClassOutermost = NULL, FName InstanceName = NAME_None) override;
@@ -170,7 +170,7 @@ public:
 	 *
 	 * @param PropertyNodePtr	The property node to make a handle for.
 	 */
-	TSharedRef<IPropertyHandle> GetPropertyHandle(TSharedPtr<FPropertyNode> PropertyNodePtr);
+	TSharedRef<IPropertyHandle> GetPropertyHandle(TSharedPtr<FPropertyNode> PropertyNodePtr) const;
 
 	/**
 	 * Adds an external property root node to the list of root nodes that the details new needs to manage
diff --git a/Engine/Source/Editor/PropertyEditor/Private/SDetailSingleItemRow.cpp b/Engine/Source/Editor/PropertyEditor/Private/SDetailSingleItemRow.cpp
--- a/Engine/Source/Editor/PropertyEditor/Private/SDetailSingleItemRow.cpp
+++ b/Engine/Source/Editor/PropertyEditor/Private/SDetailSingleItemRow.cpp
@@ -303,11 +303,12 @@ void SDetailSingleItemRow::Construct( const FArguments& InArgs, FDetailLayoutCus
 					Row.ValueWidget.Widget
 				];
 
-			ValueWidget = CreateExtensionWidget(ValueWidget.ToSharedRef(), *Customization, InOwnerTreeNode);
+			TSharedRef<SWidget> ExtensionWidget = CreateExtensionWidget(ValueWidget.ToSharedRef(), *Customization, InOwnerTreeNode);
 
 			if(Row.IsEnabledAttr.IsBound())
 			{
 				ValueWidget->SetEnabled(Row.IsEnabledAttr);
+				ExtensionWidget->SetEnabled(Row.IsEnabledAttr);
 			}
 
 			TSharedRef<SWidget> KeyFrameButton = CreateKeyframeButton(*Customization, InOwnerTreeNode);
@@ -431,21 +432,28 @@ void SDetailSingleItemRow::Construct( const FArguments& InArgs, FDetailLayoutCus
 					.Value(ColumnSizeData.RightColumnWidth)
 					.OnSlotResized(ColumnSizeData.OnWidthChanged)
 					[
-						SNew(SHorizontalBox)
-						+ SHorizontalBox::Slot()
+						SNew(SBox)
+						.IsEnabled(IsPropertyEditingEnabled)
 						[
-							SAssignNew(HBox, SHorizontalBox)
-							+ SHorizontalBox::Slot()
-							.Padding(DetailWidgetConstants::RightRowPadding)
-							.HAlign(Row.ValueWidget.HorizontalAlignment)
-							.VAlign(Row.ValueWidget.VerticalAlignment)
+							SNew(SHorizontalBox)
+							+SHorizontalBox::Slot()
+							.FillWidth(1.0f)
 							[
-								SNew(SBox)
-								.IsEnabled(IsPropertyEditingEnabled)
+								SNew(SHorizontalBox)
+								+ SHorizontalBox::Slot()
+								.Padding(DetailWidgetConstants::RightRowPadding)
+								.HAlign(Row.ValueWidget.HorizontalAlignment)
+								.VAlign(Row.ValueWidget.VerticalAlignment)
 								[
 									ValueWidget.ToSharedRef()
 								]
 							]
+							+SHorizontalBox::Slot()
+							.AutoWidth()
+							.VAlign(VAlign_Center)
+							[
+								ExtensionWidget
+							]
 						]
 					];
 				Widget = Splitter;
@@ -721,6 +729,8 @@ const FSlateBrush* SDetailSingleItemRow::GetBorderImage() const
 
 TSharedRef<SWidget> SDetailSingleItemRow::CreateExtensionWidget(TSharedRef<SWidget> ValueWidget, FDetailLayoutCustomization& InCustomization, TSharedRef<FDetailTreeNode> InTreeNode)
 {
+	TSharedPtr<SWidget> ExtensionWidget = SNullWidget::NullWidget;
+
 	if(InTreeNode->GetParentCategory().IsValid())
 	{
 		IDetailsViewPrivate* DetailsView = InTreeNode->GetDetailsView();
@@ -733,23 +743,13 @@ TSharedRef<SWidget> SDetailSingleItemRow::CreateExtensionWidget(TSharedRef<SWidg
 			UClass* ObjectClass = InCustomization.GetPropertyNode()->FindObjectItemParent()->GetObjectBaseClass();
 			if(Handle->IsValidHandle() && ExtensionHandler->IsPropertyExtendable(ObjectClass, *Handle))
 			{
-				ValueWidget = SNew(SHorizontalBox)
-					+ SHorizontalBox::Slot()
-					.FillWidth(1.0f)
-					[
-						ValueWidget
-					]
-
-					+ SHorizontalBox::Slot()
-					.AutoWidth()
-					[
-						ExtensionHandler->GenerateExtensionWidget(ObjectClass, Handle)
-					];
+				FDetailLayoutBuilderImpl& DetailLayout = OwnerTreeNode.Pin()->GetParentCategory()->GetParentLayoutImpl();
+				ExtensionWidget = ExtensionHandler->GenerateExtensionWidget(DetailLayout, ObjectClass, Handle);
 			}
 		}
 	}
 
-	return ValueWidget;
+	return ExtensionWidget.ToSharedRef();
 }
 
 TSharedRef<SWidget> SDetailSingleItemRow::CreateKeyframeButton( FDetailLayoutCustomization& InCustomization, TSharedRef<FDetailTreeNode> InTreeNode )
diff --git a/Engine/Source/Editor/PropertyEditor/Public/DetailLayoutBuilder.h b/Engine/Source/Editor/PropertyEditor/Public/DetailLayoutBuilder.h
--- a/Engine/Source/Editor/PropertyEditor/Public/DetailLayoutBuilder.h
+++ b/Engine/Source/Editor/PropertyEditor/Public/DetailLayoutBuilder.h
@@ -159,7 +159,7 @@ public:
 		To access StaticArray at index 2 from Struct2 in MyActor, your path would be MyStruct.StaticArray[2]" and your instance name is "Struct2"
 		To access MyFloat in MyActor you can just pass in "MyFloat" because the name of the property is unambiguous
 	 */
-	virtual TSharedRef<IPropertyHandle> GetProperty( const FName PropertyPath, const UStruct* ClassOutermost = NULL, FName InstanceName = NAME_None )  = 0;
+	virtual TSharedRef<IPropertyHandle> GetProperty( const FName PropertyPath, const UStruct* ClassOutermost = NULL, FName InstanceName = NAME_None ) const = 0;
 
 	/**
 	 * Gets the top level property, for showing the warning for experimental or early access class
diff --git a/Engine/Source/Editor/PropertyEditor/Public/IDetailPropertyExtensionHandler.h b/Engine/Source/Editor/PropertyEditor/Public/IDetailPropertyExtensionHandler.h
--- a/Engine/Source/Editor/PropertyEditor/Public/IDetailPropertyExtensionHandler.h
+++ b/Engine/Source/Editor/PropertyEditor/Public/IDetailPropertyExtensionHandler.h
@@ -6,6 +6,7 @@
 #include "Widgets/SWidget.h"
 
 class IPropertyHandle;
+class IDetailLayoutBuilder;
 
 class IDetailPropertyExtensionHandler
 {
@@ -14,5 +15,17 @@ public:
 
 	virtual bool IsPropertyExtendable(const UClass* InObjectClass, const class IPropertyHandle& PropertyHandle) const = 0;
 
-	virtual TSharedRef<SWidget> GenerateExtensionWidget(const UClass* InObjectClass, TSharedPtr<IPropertyHandle> PropertyHandle) = 0;
+	UE_DEPRECATED(4.24, "Please use the overload that takes a IDetailLayoutBuilder")
+	virtual TSharedRef<SWidget> GenerateExtensionWidget(const UClass* InObjectClass, TSharedPtr<IPropertyHandle> PropertyHandle) 
+	{ 
+		return SNullWidget::NullWidget;
+	}
+
+	virtual TSharedRef<SWidget> GenerateExtensionWidget(const IDetailLayoutBuilder& InDetailBuilder, const UClass* InObjectClass, TSharedPtr<IPropertyHandle> PropertyHandle)
+	{
+	PRAGMA_DISABLE_DEPRECATION_WARNINGS
+		// Call old deprecated path for back-compat
+		return GenerateExtensionWidget(InObjectClass, PropertyHandle);
+	PRAGMA_ENABLE_DEPRECATION_WARNINGS
+	}
 };
diff --git a/Engine/Source/Editor/SequenceRecorder/Private/SClassPropertyRecorderSettings.cpp b/Engine/Source/Editor/SequenceRecorder/Private/SClassPropertyRecorderSettings.cpp
--- a/Engine/Source/Editor/SequenceRecorder/Private/SClassPropertyRecorderSettings.cpp
+++ b/Engine/Source/Editor/SequenceRecorder/Private/SClassPropertyRecorderSettings.cpp
@@ -143,7 +143,7 @@ bool SClassPropertyRecorderSettings::IsPropertyExtendable(const UClass* InObject
 	return PropertyHandle.GetProperty()->HasAnyPropertyFlags(CPF_Interp) && FMovieSceneMultiPropertyRecorder::CanPropertyBeRecorded(*PropertyHandle.GetProperty());
 }
 
-TSharedRef<SWidget> SClassPropertyRecorderSettings::GenerateExtensionWidget(const UClass* InObjectClass, TSharedPtr<IPropertyHandle> PropertyHandle)
+TSharedRef<SWidget> SClassPropertyRecorderSettings::GenerateExtensionWidget(const IDetailLayoutBuilder& InDetailLayoutBuilder, const UClass* InObjectClass, TSharedPtr<IPropertyHandle> PropertyHandle)
 {
 	ECheckBoxState InitialState = ECheckBoxState::Unchecked;
 	FName PropertyName = *PropertyHandle->GeneratePathToProperty();
diff --git a/Engine/Source/Editor/SequenceRecorder/Private/SClassPropertyRecorderSettings.h b/Engine/Source/Editor/SequenceRecorder/Private/SClassPropertyRecorderSettings.h
--- a/Engine/Source/Editor/SequenceRecorder/Private/SClassPropertyRecorderSettings.h
+++ b/Engine/Source/Editor/SequenceRecorder/Private/SClassPropertyRecorderSettings.h
@@ -41,7 +41,7 @@ private:
 
 	/** IDetailPropertyExtensionHandler interface */
 	virtual bool IsPropertyExtendable(const UClass* InObjectClass, const class IPropertyHandle& PropertyHandle) const override;
-	virtual TSharedRef<SWidget> GenerateExtensionWidget(const UClass* InObjectClass, TSharedPtr<IPropertyHandle> PropertyHandle) override;
+	virtual TSharedRef<SWidget> GenerateExtensionWidget(const IDetailLayoutBuilder& InDetailLayoutBuilder, const UClass* InObjectClass, TSharedPtr<IPropertyHandle> PropertyHandle) override;
 
 private:
 	/** Handle to the property we are editing */
diff --git a/Engine/Source/Editor/Sequencer/Public/ISequencer.h b/Engine/Source/Editor/Sequencer/Public/ISequencer.h
--- a/Engine/Source/Editor/Sequencer/Public/ISequencer.h
+++ b/Engine/Source/Editor/Sequencer/Public/ISequencer.h
@@ -28,6 +28,7 @@ class FSequencerKeyCollection;
 class UMovieSceneSequence;
 class UMovieSceneSubSection;
 class IDetailsView;
+class IKeyArea;
 enum class EMapChangeType : uint8;
 
 /**
@@ -474,6 +475,12 @@ public:
 	/** Gets the currently selected folders. */
 	virtual void GetSelectedFolders(TArray<UMovieSceneFolder*>& OutSelectedFolders) = 0;
 
+	/** Gets the currently selected key areas */
+	virtual void GetSelectedKeyAreas(TArray<const IKeyArea*>& OutSelectedKeyAreas) = 0;
+
+	/** Gets the currently selected Object Guids*/
+	virtual void GetSelectedObjects(TArray<FGuid>& OutSelectedObjects) = 0;
+
 	/** Selects an object by GUID */
 	virtual void SelectObject(FGuid ObjectBinding) = 0;
 
@@ -486,6 +493,9 @@ public:
 	/** Selects property tracks by property path */
 	virtual void SelectByPropertyPaths(const TArray<FString>& InPropertyPaths) = 0;
 
+	/** Selects nodes by key areas */
+	virtual void SelectByKeyAreas(const TArray<IKeyArea>& InKeyAreas, bool bSelectParentInstead, bool bSelect) = 0;
+
 	/** Empties the current selection. */
 	virtual void EmptySelection() = 0;
 
diff --git a/Engine/Source/Editor/SkeletalMeshEditor/Private/SkeletalMeshEditor.cpp b/Engine/Source/Editor/SkeletalMeshEditor/Private/SkeletalMeshEditor.cpp
--- a/Engine/Source/Editor/SkeletalMeshEditor/Private/SkeletalMeshEditor.cpp
+++ b/Engine/Source/Editor/SkeletalMeshEditor/Private/SkeletalMeshEditor.cpp
@@ -23,7 +23,7 @@
 #include "Rendering/SkeletalMeshModel.h"
 
 #include "Animation/DebugSkelMeshComponent.h"
-#include "Assets/ClothingAsset.h"
+#include "ClothingAsset.h"
 #include "SCreateClothingSettingsPanel.h"
 #include "ClothingSystemEditorInterfaceModule.h"
 #include "Preferences/PersonaOptions.h"
@@ -44,6 +44,8 @@
 #include "EditorFramework/AssetImportData.h"
 #include "Factories/FbxSkeletalMeshImportData.h"
 
+#include "Misc/MessageDialog.h"
+
 const FName SkeletalMeshEditorAppIdentifier = FName(TEXT("SkeletalMeshEditorApp"));
 
 const FName SkeletalMeshEditorModes::SkeletalMeshEditorMode(TEXT("SkeletalMeshEditorMode"));
@@ -78,6 +80,112 @@ FSkeletalMeshEditor::~FSkeletalMeshEditor()
 	}
 }
 
+bool IsReductionParentBaseLODUseSkeletalMeshBuildWorkflow(USkeletalMesh* SkeletalMesh, int32 TestLODIndex)
+{
+	FSkeletalMeshLODInfo* LODInfo = SkeletalMesh->GetLODInfo(TestLODIndex);
+	if (LODInfo == nullptr || !SkeletalMesh->GetImportedModel() || !SkeletalMesh->GetImportedModel()->LODModels.IsValidIndex(TestLODIndex))
+	{
+		return false;
+	}
+
+	if (SkeletalMesh->GetImportedModel()->LODModels[TestLODIndex].RawSkeletalMeshBulkData.IsBuildDataAvailable())
+	{
+		return true;
+	}
+	if (LODInfo->bHasBeenSimplified || SkeletalMesh->IsReductionActive(TestLODIndex))
+	{
+		int32 ReduceBaseLOD = LODInfo->ReductionSettings.BaseLOD;
+		if (ReduceBaseLOD < TestLODIndex)
+		{
+			return IsReductionParentBaseLODUseSkeletalMeshBuildWorkflow(SkeletalMesh, ReduceBaseLOD);
+		}
+	}
+	return false;
+}
+
+bool FSkeletalMeshEditor::OnRequestClose()
+{
+	bool bAllowClose = true;
+
+	if (PersonaToolkit.IsValid() && SkeletalMesh)
+	{
+		bool bHaveModifiedLOD = false;
+		for (int32 LODIndex = 0; LODIndex < SkeletalMesh->GetLODNum(); ++LODIndex)
+		{
+			FSkeletalMeshLODInfo* LODInfo = SkeletalMesh->GetLODInfo(LODIndex);
+			if (LODInfo == nullptr || !SkeletalMesh->GetImportedModel() || !SkeletalMesh->GetImportedModel()->LODModels.IsValidIndex(LODIndex))
+			{
+				continue;
+			}
+			
+			//Do not prevent exiting if we are not using the skeletal mesh build workflow
+			if (!SkeletalMesh->GetImportedModel()->LODModels[LODIndex].RawSkeletalMeshBulkData.IsBuildDataAvailable())
+			{
+				if (!LODInfo->bHasBeenSimplified && !SkeletalMesh->IsReductionActive(LODIndex))
+				{
+					continue;
+				}
+				//Do not prevent exit if the generated LOD is base on a LODModel not using the skeletalmesh build workflow
+				int32 ReduceBaseLOD = LODInfo->ReductionSettings.BaseLOD;
+				if(!IsReductionParentBaseLODUseSkeletalMeshBuildWorkflow(SkeletalMesh, ReduceBaseLOD))
+				{
+					continue;
+				}
+				
+			}
+
+			bool bValidLODSettings = false;
+			if (SkeletalMesh->LODSettings != nullptr)
+			{
+				const int32 NumSettings = FMath::Min(SkeletalMesh->LODSettings->GetNumberOfSettings(), SkeletalMesh->GetLODNum());
+				if (LODIndex < NumSettings)
+				{
+					bValidLODSettings = true;
+				}
+			}
+			const FSkeletalMeshLODGroupSettings* SkeletalMeshLODGroupSettings = bValidLODSettings ? &SkeletalMesh->LODSettings->GetSettingsForLODLevel(LODIndex) : nullptr;
+
+			FGuid BuildGUID = LODInfo->ComputeDeriveDataCacheKey(SkeletalMeshLODGroupSettings);
+			if (LODInfo->BuildGUID != BuildGUID)
+			{
+				bHaveModifiedLOD = true;
+				break;
+			}
+			FString BuildStringID = SkeletalMesh->GetImportedModel()->LODModels[LODIndex].GetLODModelDeriveDataKey();
+			if (SkeletalMesh->GetImportedModel()->LODModels[LODIndex].BuildStringID != BuildStringID)
+			{
+				bHaveModifiedLOD = true;
+				break;
+			}
+		}
+
+		if (bHaveModifiedLOD)
+		{
+			// find out the user wants to do with this dirty material
+			EAppReturnType::Type OkCancelReply = FMessageDialog::Open(
+				EAppMsgType::OkCancel,
+				FText::Format(LOCTEXT("SkeletalMeshEditorShouldApplyLODChanges", "We have to apply level of detail changes to {0} before exiting the skeletal mesh editor."), FText::FromString(PersonaToolkit->GetMesh()->GetName()))
+			);
+
+			switch (OkCancelReply)
+			{
+			case EAppReturnType::Ok:
+				{
+					SkeletalMesh->MarkPackageDirty();
+					SkeletalMesh->PostEditChange();
+					bAllowClose = true;
+				}
+				break;
+			case EAppReturnType::Cancel:
+				// Don't exit.
+				bAllowClose = false;
+				break;
+			}
+		}
+	}
+	return bAllowClose;
+}
+
 void FSkeletalMeshEditor::RegisterTabSpawners(const TSharedRef<class FTabManager>& InTabManager)
 {
 	WorkspaceMenuCategory = InTabManager->AddLocalWorkspaceMenuCategory(LOCTEXT("WorkspaceMenu_SkeletalMeshEditor", "Skeletal Mesh Editor"));
@@ -546,8 +654,7 @@ void FSkeletalMeshEditor::OnRemoveSectionFromLodAndBelowMenuItemClicked(int32 Lo
 		const FSkeletalMeshLODInfo* CurrentSkeletalMeshLODInfo = SkeletalMesh->GetLODInfo(GenerateLodIndex);
 		if (CurrentSkeletalMeshLODInfo != nullptr && CurrentSkeletalMeshLODInfo->bHasBeenSimplified && BaseLodIndexes.Contains(CurrentSkeletalMeshLODInfo->ReductionSettings.BaseLOD))
 		{
-			bool bRestoreClothing = true;
-			FLODUtilities::SimplifySkeletalMeshLOD(UpdateContext, GenerateLodIndex, true, bRestoreClothing);
+			FLODUtilities::SimplifySkeletalMeshLOD(UpdateContext, GenerateLodIndex, true);
 			BaseLodIndexes.Add(GenerateLodIndex);
 		}
 	}
@@ -569,12 +676,12 @@ void FSkeletalMeshEditor::FillApplyClothingAssetMenu(FMenuBuilder& MenuBuilder,
 	{
 		for(UClothingAssetBase* BaseAsset : Mesh->MeshClothingAssets)
 		{
-			UClothingAsset* ClothAsset = CastChecked<UClothingAsset>(BaseAsset);
+			UClothingAssetCommon* ClothAsset = CastChecked<UClothingAssetCommon>(BaseAsset);
 
 			FUIAction Action;
 			Action.CanExecuteAction = FCanExecuteAction::CreateSP(this, &FSkeletalMeshEditor::CanApplyClothing, InLodIndex, InSectionIndex);
 
-			const int32 NumClothLods = ClothAsset->LodData.Num();
+			const int32 NumClothLods = ClothAsset->GetNumLods();
 			for(int32 ClothLodIndex = 0; ClothLodIndex < NumClothLods; ++ClothLodIndex)
 			{
 				Action.ExecuteAction = FExecuteAction::CreateSP(this, &FSkeletalMeshEditor::OnApplyClothingAssetClicked, BaseAsset, InLodIndex, InSectionIndex, ClothLodIndex);
@@ -668,7 +775,7 @@ void FSkeletalMeshEditor::OnCreateClothingAssetMenuItemClicked(FSkeletalMeshClot
 						{
 							SectionIndex = i;
 							AssetLodIndex = SkelLod.Sections[i].ClothingData.AssetLodIndex;
-							TargetAssetPtr->UnbindFromSkeletalMesh(Mesh, Params.TargetLod);
+							RemoveClothing(Params.TargetLod, SectionIndex);
 							break;
 						}
 					}
@@ -786,13 +893,54 @@ void FSkeletalMeshEditor::ApplyClothing(UClothingAssetBase* InAsset, int32 InLod
 {
 	USkeletalMesh* Mesh = GetPersonaToolkit()->GetPreviewMesh();
 
-	if(UClothingAsset* ClothingAsset = Cast<UClothingAsset>(InAsset))
+	if (Mesh == nullptr || Mesh->GetImportedModel() == nullptr || !Mesh->GetImportedModel()->LODModels.IsValidIndex(InLodIndex))
 	{
-		ClothingAsset->BindToSkeletalMesh(Mesh, InLodIndex, InSectionIndex, InClothingLod);
+		return;
 	}
-	else
+
+	FSkeletalMeshLODModel& LODModel = Mesh->GetImportedModel()->LODModels[InLodIndex];
+	const FSkelMeshSection& Section = LODModel.Sections[InSectionIndex];
+	FScopedSuspendAlternateSkinWeightPreview ScopedSuspendAlternateSkinnWeightPreview(Mesh);
 	{
-		RemoveClothing(InLodIndex, InSectionIndex);
+		FScopedSkeletalMeshPostEditChange ScopedPostEditChange(Mesh);
+		FScopedTransaction Transaction(LOCTEXT("SkeletalMeshEditorApplyClothingTransaction", "Persona editor: Apply Section Cloth"));
+		Mesh->Modify();
+
+		FSkelMeshSourceSectionUserData& OriginalSectionData = LODModel.UserSectionsData.FindOrAdd(Section.OriginalDataSectionIndex);
+		auto ClearOriginalSectionUserData = [&OriginalSectionData]()
+		{
+			OriginalSectionData.CorrespondClothAssetIndex = INDEX_NONE;
+			OriginalSectionData.ClothingData.AssetGuid = FGuid();
+			OriginalSectionData.ClothingData.AssetLodIndex = INDEX_NONE;
+		};
+		if (UClothingAssetCommon* ClothingAsset = Cast<UClothingAssetCommon>(InAsset))
+		{
+			// Look for a currently bound asset an unbind it if necessary first
+			if (UClothingAssetBase* CurrentAsset = Mesh->GetSectionClothingAsset(InLodIndex, InSectionIndex))
+			{
+				CurrentAsset->UnbindFromSkeletalMesh(Mesh, InLodIndex);
+				ClearOriginalSectionUserData();
+			}
+
+			if (ClothingAsset->BindToSkeletalMesh(Mesh, InLodIndex, InSectionIndex, InClothingLod))
+			{
+				//Successful bind so set the SectionUserData
+				int32 AssetIndex = INDEX_NONE;
+				check(Mesh->MeshClothingAssets.Find(ClothingAsset, AssetIndex));
+				OriginalSectionData.CorrespondClothAssetIndex = AssetIndex;
+				OriginalSectionData.ClothingData.AssetGuid = ClothingAsset->GetAssetGuid();
+				OriginalSectionData.ClothingData.AssetLodIndex = InClothingLod;
+			}
+		}
+		else if (Mesh)
+		{
+			//User set none, so unbind anything that is bind
+			if (UClothingAssetBase* CurrentAsset = Mesh->GetSectionClothingAsset(InLodIndex, InSectionIndex))
+			{
+				CurrentAsset->UnbindFromSkeletalMesh(Mesh, InLodIndex);
+				ClearOriginalSectionUserData();
+			}
+		}
 	}
 }
 
@@ -800,11 +948,31 @@ void FSkeletalMeshEditor::RemoveClothing(int32 InLodIndex, int32 InSectionIndex)
 {
 	USkeletalMesh* Mesh = GetPersonaToolkit()->GetPreviewMesh();
 
-	if(Mesh)
+	if (Mesh->GetImportedModel() == nullptr || !Mesh->GetImportedModel()->LODModels.IsValidIndex(InLodIndex))
 	{
-		if(UClothingAssetBase* CurrentAsset = Mesh->GetSectionClothingAsset(InLodIndex, InSectionIndex))
+		return;
+	}
+
+	FSkeletalMeshLODModel& LODModel = Mesh->GetImportedModel()->LODModels[InLodIndex];
+	const FSkelMeshSection& Section = LODModel.Sections[InSectionIndex];
+	FScopedSuspendAlternateSkinWeightPreview ScopedSuspendAlternateSkinnWeightPreview(Mesh);
+	{
+		FScopedSkeletalMeshPostEditChange ScopedPostEditChange(Mesh);
+		FScopedTransaction Transaction(LOCTEXT("SkeletalMeshEditorRemoveClothingTransaction", "Persona editor: Remove Section Cloth"));
+		Mesh->Modify();
+
+		FSkelMeshSourceSectionUserData& OriginalSectionData = LODModel.UserSectionsData.FindOrAdd(Section.OriginalDataSectionIndex);
+		auto ClearOriginalSectionUserData = [&OriginalSectionData]()
+		{
+			OriginalSectionData.CorrespondClothAssetIndex = INDEX_NONE;
+			OriginalSectionData.ClothingData.AssetGuid = FGuid();
+			OriginalSectionData.ClothingData.AssetLodIndex = INDEX_NONE;
+		};
+		// Look for a currently bound asset an unbind it if necessary first
+		if (UClothingAssetBase* CurrentAsset = Mesh->GetSectionClothingAsset(InLodIndex, InSectionIndex))
 		{
 			CurrentAsset->UnbindFromSkeletalMesh(Mesh, InLodIndex);
+			ClearOriginalSectionUserData();
 		}
 	}
 }
@@ -883,92 +1051,92 @@ UObject* FSkeletalMeshEditor::HandleGetAsset()
 
 bool FSkeletalMeshEditor::HandleReimportMeshInternal(int32 SourceFileIndex /*= INDEX_NONE*/, bool bWithNewFile /*= false*/)
 {
-	// Make sure we clear any skin weight profile currently being previewed
-	UDebugSkelMeshComponent* MeshComponent = GetPersonaToolkit()->GetPreviewScene()->GetPreviewMeshComponent();
-	const FName ProfileName = MeshComponent ? MeshComponent->GetCurrentSkinWeightProfileName() : NAME_None;
-	if (MeshComponent)
-	{
-		MeshComponent->ClearSkinWeightProfile();
-	}
-
-	// Reimport the asset
-	const bool bResult = FReimportManager::Instance()->Reimport(SkeletalMesh, true, true, TEXT(""), nullptr, SourceFileIndex, bWithNewFile);
-
-	// Refresh skeleton tree
-	SkeletonTree->Refresh();
-
-	// Re-set skin weigh profile
-	if (MeshComponent && (ProfileName != NAME_None))
+	FScopedSuspendAlternateSkinWeightPreview ScopedSuspendAlternateSkinnWeightPreview(SkeletalMesh);
+	bool bResult = false;
 	{
-		MeshComponent->SetSkinWeightProfile(ProfileName);
+		FScopedSkeletalMeshPostEditChange ScopedPostEditChange(SkeletalMesh);
+		// Reimport the asset
+		bResult = FReimportManager::Instance()->Reimport(SkeletalMesh, true, true, TEXT(""), nullptr, SourceFileIndex, bWithNewFile);
+		// Refresh skeleton tree
+		SkeletonTree->Refresh();
 	}
-
 	return bResult;
 }
 
 void FSkeletalMeshEditor::HandleReimportMesh(int32 SourceFileIndex /*= INDEX_NONE*/)
 {
+	FScopedSuspendAlternateSkinWeightPreview ScopedSuspendAlternateSkinnWeightPreview(SkeletalMesh);
+
 	HandleReimportMeshInternal(SourceFileIndex, false);
 }
 
 void FSkeletalMeshEditor::HandleReimportMeshWithNewFile(int32 SourceFileIndex /*= INDEX_NONE*/)
 {
+	FScopedSuspendAlternateSkinWeightPreview ScopedSuspendAlternateSkinnWeightPreview(SkeletalMesh);
+
 	HandleReimportMeshInternal(SourceFileIndex, true);
 }
 
 void ReimportAllCustomLODs(USkeletalMesh* SkeletalMesh, UDebugSkelMeshComponent* PreviewMeshComponent, bool bWithNewFile)
 {
-	//Find the dependencies of the generated LOD
-	TArray<bool> Dependencies;
-	Dependencies.AddZeroed(SkeletalMesh->GetLODNum());
-	//Avoid making LOD 0 to true in the dependencies since everything that should be regenerate base on LOD 0 is already regenerate at this point.
-	//But we need to regenerate every generated LOD base on any re-import custom LOD
-	//Reimport all custom LODs
-	for (int32 LodIndex = 1; LodIndex < SkeletalMesh->GetLODNum(); ++LodIndex)
+	FScopedSuspendAlternateSkinWeightPreview ScopedSuspendAlternateSkinnWeightPreview(SkeletalMesh);
 	{
-		//Do not reimport LOD that was re-import with the base mesh
-		if (SkeletalMesh->GetLODInfo(LodIndex)->bImportWithBaseMesh)
-		{
-			continue;
-		}
-		if (SkeletalMesh->GetLODInfo(LodIndex)->bHasBeenSimplified == false)
+		FScopedSkeletalMeshPostEditChange ScopedPostEditChange(SkeletalMesh);
+
+		//Find the dependencies of the generated LOD
+		TArray<bool> Dependencies;
+		Dependencies.AddZeroed(SkeletalMesh->GetLODNum());
+		//Avoid making LOD 0 to true in the dependencies since everything that should be regenerate base on LOD 0 is already regenerate at this point.
+		//But we need to regenerate every generated LOD base on any re-import custom LOD
+		//Reimport all custom LODs
+		for (int32 LodIndex = 1; LodIndex < SkeletalMesh->GetLODNum(); ++LodIndex)
 		{
-			FString SourceFilenameBackup = SkeletalMesh->GetLODInfo(LodIndex)->SourceImportFilename;
-			if (bWithNewFile)
+			//Do not reimport LOD that was re-import with the base mesh
+			if (SkeletalMesh->GetLODInfo(LodIndex)->bImportWithBaseMesh)
 			{
-				SkeletalMesh->GetLODInfo(LodIndex)->SourceImportFilename.Empty();
+				continue;
 			}
-
-			if (!FbxMeshUtils::ImportMeshLODDialog(SkeletalMesh, LodIndex))
+			if (SkeletalMesh->GetLODInfo(LodIndex)->bHasBeenSimplified == false)
 			{
+				FString SourceFilenameBackup = SkeletalMesh->GetLODInfo(LodIndex)->SourceImportFilename;
 				if (bWithNewFile)
 				{
-					SkeletalMesh->GetLODInfo(LodIndex)->SourceImportFilename = SourceFilenameBackup;
+					SkeletalMesh->GetLODInfo(LodIndex)->SourceImportFilename.Empty();
+				}
+
+				if (!FbxMeshUtils::ImportMeshLODDialog(SkeletalMesh, LodIndex, false))
+				{
+					if (bWithNewFile)
+					{
+						SkeletalMesh->GetLODInfo(LodIndex)->SourceImportFilename = SourceFilenameBackup;
+					}
+				}
+				else
+				{
+					Dependencies[LodIndex] = true;
 				}
 			}
-			else
+			else if (Dependencies[SkeletalMesh->GetLODInfo(LodIndex)->ReductionSettings.BaseLOD])
 			{
+				//Regenerate the LOD
+				FSkeletalMeshUpdateContext UpdateContext;
+				UpdateContext.SkeletalMesh = SkeletalMesh;
+				UpdateContext.AssociatedComponents.Push(PreviewMeshComponent);
+				FLODUtilities::SimplifySkeletalMeshLOD(UpdateContext, LodIndex, false);
 				Dependencies[LodIndex] = true;
 			}
 		}
-		else if(Dependencies[SkeletalMesh->GetLODInfo(LodIndex)->ReductionSettings.BaseLOD])
-		{
-			//Regenerate the LOD
-			FSkeletalMeshUpdateContext UpdateContext;
-			UpdateContext.SkeletalMesh = SkeletalMesh;
-			UpdateContext.AssociatedComponents.Push(PreviewMeshComponent);
-			bool bRestoreClothing = true;
-			FLODUtilities::SimplifySkeletalMeshLOD(UpdateContext, LodIndex, true, bRestoreClothing);
-			Dependencies[LodIndex] = true;
-		}
 	}
 }
 
 void FSkeletalMeshEditor::HandleReimportAllMesh(int32 SourceFileIndex /*= INDEX_NONE*/)
 {
+	FScopedSuspendAlternateSkinWeightPreview ScopedSuspendAlternateSkinnWeightPreview(SkeletalMesh);
 	// Reimport the asset
 	if (SkeletalMesh)
 	{
+		FScopedSkeletalMeshPostEditChange ScopedPostEditChange(SkeletalMesh);
+
 		//Reimport base LOD
 		if (HandleReimportMeshInternal(SourceFileIndex, false))
 		{
@@ -980,9 +1148,12 @@ void FSkeletalMeshEditor::HandleReimportAllMesh(int32 SourceFileIndex /*= INDEX_
 
 void FSkeletalMeshEditor::HandleReimportAllMeshWithNewFile(int32 SourceFileIndex /*= INDEX_NONE*/)
 {
+	FScopedSuspendAlternateSkinWeightPreview ScopedSuspendAlternateSkinnWeightPreview(SkeletalMesh);
+
 	// Reimport the asset
 	if (SkeletalMesh)
 	{
+		FScopedSkeletalMeshPostEditChange ScopedPostEditChange(SkeletalMesh);
 		TArray<UObject*> ImportObjs;
 		ImportObjs.Add(SkeletalMesh);
 		if (HandleReimportMeshInternal(SourceFileIndex, true))
diff --git a/Engine/Source/Editor/SkeletonEditor/Private/SSkeletonTree.cpp b/Engine/Source/Editor/SkeletonEditor/Private/SSkeletonTree.cpp
--- a/Engine/Source/Editor/SkeletonEditor/Private/SSkeletonTree.cpp
+++ b/Engine/Source/Editor/SkeletonEditor/Private/SSkeletonTree.cpp
@@ -906,12 +906,20 @@ bool GetSourceNameFromItem(TSharedPtr<ISkeletonTreeItem> SourceBone, FName& OutN
 void SSkeletonTree::FillVirtualBoneSubmenu(FMenuBuilder& MenuBuilder, TArray<TSharedPtr<ISkeletonTreeItem>> SourceBones)
 {
 	const bool bShowVirtualBones = false;
-	TSharedRef<SWidget> MenuContent = SNew(SBoneTreeMenu)
+	TSharedRef<SBoneTreeMenu> MenuContent = SNew(SBoneTreeMenu)
 		.bShowVirtualBones(false)
 		.Title(LOCTEXT("TargetBonePickerTitle", "Pick Target Bone..."))
 		.OnBoneSelectionChanged(this, &SSkeletonTree::OnVirtualTargetBonePicked, SourceBones)
 		.OnGetReferenceSkeleton(this, &SSkeletonTree::OnGetReferenceSkeleton);
 	MenuBuilder.AddWidget(MenuContent, FText::GetEmpty(), true);
+
+	MenuContent->RegisterActiveTimer(0.0f, FWidgetActiveTimerDelegate::CreateLambda(
+		[FilterTextBox = MenuContent->GetFilterTextWidget()](double, float)
+		{
+			FSlateApplication::Get().SetKeyboardFocus(FilterTextBox);
+			return EActiveTimerReturnType::Stop;
+		}
+	));
 }
 
 void SSkeletonTree::OnVirtualTargetBonePicked(FName TargetBoneName, TArray<TSharedPtr<ISkeletonTreeItem>> SourceBones)
diff --git a/Engine/Source/Editor/UMGEditor/Private/Details/DetailWidgetExtensionHandler.cpp b/Engine/Source/Editor/UMGEditor/Private/Details/DetailWidgetExtensionHandler.cpp
--- a/Engine/Source/Editor/UMGEditor/Private/Details/DetailWidgetExtensionHandler.cpp
+++ b/Engine/Source/Editor/UMGEditor/Private/Details/DetailWidgetExtensionHandler.cpp
@@ -42,7 +42,7 @@ bool FDetailWidgetExtensionHandler::IsPropertyExtendable(const UClass* InObjectC
 	return false;
 }
 
-TSharedRef<SWidget> FDetailWidgetExtensionHandler::GenerateExtensionWidget(const UClass* InObjectClass, TSharedPtr<IPropertyHandle> InPropertyHandle)
+TSharedRef<SWidget> FDetailWidgetExtensionHandler::GenerateExtensionWidget(const IDetailLayoutBuilder& InDetailBuilder, const UClass* InObjectClass, TSharedPtr<IPropertyHandle> InPropertyHandle)
 {
 	UProperty* Property = InPropertyHandle->GetProperty();
 	FString DelegateName = Property->GetName() + "Delegate";
diff --git a/Engine/Source/Editor/UMGEditor/Private/Details/DetailWidgetExtensionHandler.h b/Engine/Source/Editor/UMGEditor/Private/Details/DetailWidgetExtensionHandler.h
--- a/Engine/Source/Editor/UMGEditor/Private/Details/DetailWidgetExtensionHandler.h
+++ b/Engine/Source/Editor/UMGEditor/Private/Details/DetailWidgetExtensionHandler.h
@@ -16,7 +16,7 @@ public:
 
 	virtual bool IsPropertyExtendable(const UClass* InObjectClass, const class IPropertyHandle& PropertyHandle) const override;
 
-	virtual TSharedRef<SWidget> GenerateExtensionWidget(const UClass* InObjectClass, TSharedPtr<IPropertyHandle> PropertyHandle) override;
+	virtual TSharedRef<SWidget> GenerateExtensionWidget(const IDetailLayoutBuilder& InDetailBuilder, const UClass* InObjectClass, TSharedPtr<IPropertyHandle> PropertyHandle) override;
 
 private:
 	TWeakPtr<class FWidgetBlueprintEditor> BlueprintEditor;
diff --git a/Engine/Source/Editor/UnrealEd/Private/Animation/DebugSkelMeshComponent.cpp b/Engine/Source/Editor/UnrealEd/Private/Animation/DebugSkelMeshComponent.cpp
--- a/Engine/Source/Editor/UnrealEd/Private/Animation/DebugSkelMeshComponent.cpp
+++ b/Engine/Source/Editor/UnrealEd/Private/Animation/DebugSkelMeshComponent.cpp
@@ -17,7 +17,7 @@
 #include "Rendering/SkeletalMeshRenderData.h"
 #include "Rendering/SkeletalMeshModel.h"
 
-#include "Assets/ClothingAsset.h"
+#include "ClothingAsset.h"
 #include "ClothingSimulation.h"
 #include "Utils/ClothingMeshUtils.h"
 #include "DynamicMeshBuilder.h"
@@ -740,17 +740,15 @@ void UDebugSkelMeshComponent::RebuildClothingSectionsFixedVerts()
 
 				if(BaseAsset)
 				{
-					UClothingAsset* ConcreteAsset = Cast<UClothingAsset>(BaseAsset);
-					FClothLODData& LodData = ConcreteAsset->LodData[Section.ClothingData.AssetLodIndex];
+					UClothingAssetCommon* ConcreteAsset = Cast<UClothingAssetCommon>(BaseAsset);
+					UClothLODDataBase* LodData = ConcreteAsset->ClothLodData[Section.ClothingData.AssetLodIndex];
 
 					for(FMeshToMeshVertData& VertData : Section.ClothMappingData)
 					{
-						float TriangleDistanceMax = 0.0f;
-						TriangleDistanceMax += LodData.PhysicalMeshData.MaxDistances[VertData.SourceMeshVertIndices[0]];
-						TriangleDistanceMax += LodData.PhysicalMeshData.MaxDistances[VertData.SourceMeshVertIndices[1]];
-						TriangleDistanceMax += LodData.PhysicalMeshData.MaxDistances[VertData.SourceMeshVertIndices[2]];
-
-						if(TriangleDistanceMax == 0.0f)
+						if(LodData->PhysicalMeshData->IsFixed(
+							VertData.SourceMeshVertIndices[0],
+							VertData.SourceMeshVertIndices[1],
+							VertData.SourceMeshVertIndices[2]))
 						{
 							VertData.SourceMeshVertIndices[3] = 0xFFFF;
 						}
@@ -819,9 +817,9 @@ void UDebugSkelMeshComponent::RefreshSelectedClothingSkinnedPositions()
 
 		if(Asset)
 		{
-			UClothingAsset* ConcreteAsset = Cast<UClothingAsset>(*Asset);
+			UClothingAssetCommon* ConcreteAsset = Cast<UClothingAssetCommon>(*Asset);
 
-			if(ConcreteAsset->LodData.IsValidIndex(SelectedClothingLodForPainting))
+			if(ConcreteAsset->ClothLodData.IsValidIndex(SelectedClothingLodForPainting))
 			{
 				SkinnedSelectedClothingPositions.Reset();
 				SkinnedSelectedClothingNormals.Reset();
@@ -830,9 +828,9 @@ void UDebugSkelMeshComponent::RefreshSelectedClothingSkinnedPositions()
 				// Pass LOD0 to collect all bones
 				GetCurrentRefToLocalMatrices(RefToLocals, 0);
 
-				FClothLODData& LodData = ConcreteAsset->LodData[SelectedClothingLodForPainting];
+				UClothLODDataBase* LodData = ConcreteAsset->ClothLodData[SelectedClothingLodForPainting];
 
-				ClothingMeshUtils::SkinPhysicsMesh(ConcreteAsset->UsedBoneIndices, LodData.PhysicalMeshData, FTransform::Identity, RefToLocals.GetData(), RefToLocals.Num(), SkinnedSelectedClothingPositions, SkinnedSelectedClothingNormals);
+				ClothingMeshUtils::SkinPhysicsMesh(ConcreteAsset->UsedBoneIndices, *LodData->PhysicalMeshData, FTransform::Identity, RefToLocals.GetData(), RefToLocals.Num(), SkinnedSelectedClothingPositions, SkinnedSelectedClothingNormals);
 				RebuildCachedClothBounds();
 			}
 		}
@@ -1033,17 +1031,17 @@ FDebugSkelMeshDynamicData::FDebugSkelMeshDynamicData(UDebugSkelMeshComponent* In
 				{
 					ClothingSimDataIndexWhenPainting = ClothingAssetIndex;
 
-					if(UClothingAsset* ConcreteAsset = Cast<UClothingAsset>(BaseAsset))
+					if(UClothingAssetCommon* ConcreteAsset = Cast<UClothingAssetCommon>(BaseAsset))
 					{
-						if(ConcreteAsset->LodData.IsValidIndex(InComponent->SelectedClothingLodForPainting))
+						if(ConcreteAsset->ClothLodData.IsValidIndex(InComponent->SelectedClothingLodForPainting))
 						{
-							FClothLODData& LodData = ConcreteAsset->LodData[InComponent->SelectedClothingLodForPainting];
+							UClothLODDataBase* LodData = ConcreteAsset->ClothLodData[InComponent->SelectedClothingLodForPainting];
 
-							ClothingSimIndices = LodData.PhysicalMeshData.Indices;
+							ClothingSimIndices = LodData->PhysicalMeshData->Indices;
 
-							if(LodData.ParameterMasks.IsValidIndex(InComponent->SelectedClothingLodMaskForPainting))
+							if(LodData->ParameterMasks.IsValidIndex(InComponent->SelectedClothingLodMaskForPainting))
 							{
-								FClothParameterMask_PhysMesh& Mask = LodData.ParameterMasks[InComponent->SelectedClothingLodMaskForPainting];
+								FPointWeightMap& Mask = LodData->ParameterMasks[InComponent->SelectedClothingLodMaskForPainting];
 
 								ClothingVisiblePropertyValues = Mask.GetValueArray();
 							}
@@ -1057,3 +1055,37 @@ FDebugSkelMeshDynamicData::FDebugSkelMeshDynamicData(UDebugSkelMeshComponent* In
 	}
 }
 
+FScopedSuspendAlternateSkinWeightPreview::FScopedSuspendAlternateSkinWeightPreview(USkeletalMesh* SkeletalMesh)
+{
+	SuspendedComponentArray.Empty(2);
+	if (SkeletalMesh != nullptr)
+	{
+		// Now iterate over all skeletal mesh components and unregister them from the world, we will reregister them in the destructor
+		for (TObjectIterator<UDebugSkelMeshComponent> It; It; ++It)
+		{
+			UDebugSkelMeshComponent* DebugSKComp = *It;
+			if (DebugSKComp->SkeletalMesh == SkeletalMesh)
+			{
+				const FName ProfileName = DebugSKComp->GetCurrentSkinWeightProfileName();
+				if (ProfileName != NAME_None)
+				{
+					DebugSKComp->ClearSkinWeightProfile();
+					TTuple<UDebugSkelMeshComponent*, FName> ComponentTupple;
+					ComponentTupple.Key = DebugSKComp;
+					ComponentTupple.Value = ProfileName;
+					SuspendedComponentArray.Add(ComponentTupple);
+				}
+			}
+		}
+	}
+}
+
+FScopedSuspendAlternateSkinWeightPreview::~FScopedSuspendAlternateSkinWeightPreview()
+{
+	//Put back the skin weight profile for all editor debug component
+	for (const TTuple<UDebugSkelMeshComponent*, FName>& ComponentTupple : SuspendedComponentArray)
+	{
+		ComponentTupple.Key->SetSkinWeightProfile(ComponentTupple.Value);
+	}
+	SuspendedComponentArray.Empty();
+}
