Fix crash when applying change and there is a preview mesh influence with alternate skin
#rnx

diff --git a/Engine/Source/Editor/Persona/Private/PersonaMeshDetails.cpp b/Engine/Source/Editor/Persona/Private/PersonaMeshDetails.cpp
--- a/Engine/Source/Editor/Persona/Private/PersonaMeshDetails.cpp
+++ b/Engine/Source/Editor/Persona/Private/PersonaMeshDetails.cpp
@@ -2898,7 +2898,8 @@ FReply FPersonaMeshDetails::ApplyLODChanges(int32 LODIndex)
 	{
 		return FReply::Handled();
 	}
-
+	
+	FScopedSuspendAlternateSkinWeightPreview ScopedSuspendAlternateSkinnWeightPreview(SkelMesh);
 	{
 		FScopedSkeletalMeshPostEditChange ScopedPostEditChange(SkelMesh);
 		
@@ -2955,6 +2956,7 @@ void FPersonaMeshDetails::RegenerateOneLOD(int32 LODIndex)
 {
 	USkeletalMesh* SkelMesh = GetPersonaToolkit()->GetMesh();
 	check(SkelMesh);
+	FScopedSuspendAlternateSkinWeightPreview ScopedSuspendAlternateSkinnWeightPreview(SkelMesh);
 
 	if (SkelMesh->IsValidLODIndex(LODIndex))
 	{
@@ -3015,6 +3017,7 @@ void FPersonaMeshDetails::RegenerateDependentLODs(int32 LODIndex)
 
 	IMeshReductionModule& ReductionModule = FModuleManager::Get().LoadModuleChecked<IMeshReductionModule>("MeshReductionInterface");
 	IMeshReduction* MeshReduction = ReductionModule.GetSkeletalMeshReductionInterface();
+	FScopedSuspendAlternateSkinWeightPreview ScopedSuspendAlternateSkinnWeightPreview(SkelMesh);
 	if (MeshReduction && MeshReduction->IsSupported())
 	{
 		FScopedSkeletalMeshPostEditChange ScopedPostEditChange(SkelMesh);
@@ -3068,6 +3071,7 @@ FReply FPersonaMeshDetails::RegenerateLOD(int32 LODIndex)
 		}
 	}
 	
+	FScopedSuspendAlternateSkinWeightPreview ScopedSuspendAlternateSkinnWeightPreview(SkelMesh);
 	//Reregister scope
 	{
 		FScopedSkeletalMeshPostEditChange ScopedPostEditChange(SkelMesh);
@@ -3111,6 +3115,7 @@ FReply FPersonaMeshDetails::RemoveOneLOD(int32 LODIndex)
 			FScopedTransaction Transaction(TEXT(""), RemoveLODText, SkelMesh);
 			SkelMesh->Modify();
 
+			FScopedSuspendAlternateSkinWeightPreview ScopedSuspendAlternateSkinnWeightPreview(SkelMesh);
 			//PostEditChange scope
 			{
 				FScopedSkeletalMeshPostEditChange ScopedPostEditChange(SkelMesh);
@@ -3152,6 +3157,7 @@ void FPersonaMeshDetails::ApplyChanges()
 	USkeletalMesh* SkelMesh = GetPersonaToolkit()->GetMesh();
 	check(SkelMesh);
 
+	FScopedSuspendAlternateSkinWeightPreview ScopedSuspendAlternateSkinnWeightPreview(SkelMesh);
 	//Control the scope of the PostEditChange
 	{
 		FScopedSkeletalMeshPostEditChange ScopedPostEditChange(SkelMesh);
@@ -5294,6 +5300,7 @@ FReply FPersonaMeshDetails::OnRemoveApexFileClicked(int32 AssetIndex, IDetailLay
 		}
 	}
 
+	FScopedSuspendAlternateSkinWeightPreview ScopedSuspendAlternateSkinnWeightPreview(SkelMesh);
 	{
 		// Need to unregister our components so they shut down their current clothing simulation
 		FScopedSkeletalMeshPostEditChange ScopedPostEditChange(SkelMesh);
@@ -5448,6 +5455,7 @@ void FPersonaMeshDetails::OnClothingSelectionChanged(TSharedPtr<FClothingEntry>
 
 	FSkeletalMeshLODModel& LODModel = Mesh->GetImportedModel()->LODModels[InLodIdx];
 	const FSkelMeshSection& Section = LODModel.Sections[InSectionIdx];
+	FScopedSuspendAlternateSkinWeightPreview ScopedSuspendAlternateSkinnWeightPreview(Mesh);
 	{
 		FScopedSkeletalMeshPostEditChange ScopedPostEditChange(Mesh);
 		FScopedTransaction Transaction(LOCTEXT("PersonaOnSectionClothChangedTransaction", "Persona editor: Section cloth changed"));
diff --git a/Engine/Source/Editor/SkeletalMeshEditor/Private/SkeletalMeshEditor.cpp b/Engine/Source/Editor/SkeletalMeshEditor/Private/SkeletalMeshEditor.cpp
--- a/Engine/Source/Editor/SkeletalMeshEditor/Private/SkeletalMeshEditor.cpp
+++ b/Engine/Source/Editor/SkeletalMeshEditor/Private/SkeletalMeshEditor.cpp
@@ -900,6 +900,7 @@ void FSkeletalMeshEditor::ApplyClothing(UClothingAssetBase* InAsset, int32 InLod
 
 	FSkeletalMeshLODModel& LODModel = Mesh->GetImportedModel()->LODModels[InLodIndex];
 	const FSkelMeshSection& Section = LODModel.Sections[InSectionIndex];
+	FScopedSuspendAlternateSkinWeightPreview ScopedSuspendAlternateSkinnWeightPreview(Mesh);
 	{
 		FScopedSkeletalMeshPostEditChange ScopedPostEditChange(Mesh);
 		FScopedTransaction Transaction(LOCTEXT("SkeletalMeshEditorApplyClothingTransaction", "Persona editor: Apply Section Cloth"));
@@ -954,6 +955,7 @@ void FSkeletalMeshEditor::RemoveClothing(int32 InLodIndex, int32 InSectionIndex)
 
 	FSkeletalMeshLODModel& LODModel = Mesh->GetImportedModel()->LODModels[InLodIndex];
 	const FSkelMeshSection& Section = LODModel.Sections[InSectionIndex];
+	FScopedSuspendAlternateSkinWeightPreview ScopedSuspendAlternateSkinnWeightPreview(Mesh);
 	{
 		FScopedSkeletalMeshPostEditChange ScopedPostEditChange(Mesh);
 		FScopedTransaction Transaction(LOCTEXT("SkeletalMeshEditorRemoveClothingTransaction", "Persona editor: Remove Section Cloth"));
@@ -1047,72 +1049,17 @@ UObject* FSkeletalMeshEditor::HandleGetAsset()
 	return GetEditingObject();
 }
 
-class FScopedSuspendAlternateSkinWeightPreview
-{
-public:
-	/*
-	 * This constructor suspend the alternate skinning preview for all editor component that use the specified skeletalmesh
-	 * Parameters:
-	 * @param InSkeletalMesh - SkeletalMesh use to know which preview component we have to suspend the alternate skinning preview.
-	 */
-	FScopedSuspendAlternateSkinWeightPreview(USkeletalMesh* InSkeletalMesh);
-
-	/*
-	 * This destructor put back the preview alternate skinning
-	 */
-	~FScopedSuspendAlternateSkinWeightPreview();
-
-private:
-	TArray< TTuple<UDebugSkelMeshComponent*, FName> > SuspendedComponentArray;
-};
-
-FScopedSuspendAlternateSkinWeightPreview::FScopedSuspendAlternateSkinWeightPreview(USkeletalMesh* SkeletalMesh)
-{
-	SuspendedComponentArray.Empty(2);
-	if (SkeletalMesh != nullptr)
-	{
-		// Now iterate over all skeletal mesh components and unregister them from the world, we will reregister them in the destructor
-		for (TObjectIterator<UDebugSkelMeshComponent> It; It; ++It)
-		{
-			UDebugSkelMeshComponent* DebugSKComp = *It;
-			if (DebugSKComp->SkeletalMesh == SkeletalMesh)
-			{
-				const FName ProfileName = DebugSKComp->GetCurrentSkinWeightProfileName();
-				if (ProfileName != NAME_None)
-				{
-					DebugSKComp->ClearSkinWeightProfile();
-					TTuple<UDebugSkelMeshComponent*, FName> ComponentTupple;
-					ComponentTupple.Key = DebugSKComp;
-					ComponentTupple.Value = ProfileName;
-					SuspendedComponentArray.Add(ComponentTupple);
-				}
-			}
-		}
-	}
-}
-
-FScopedSuspendAlternateSkinWeightPreview::~FScopedSuspendAlternateSkinWeightPreview()
-{
-	//Put back the skin weight profile for all editor debug component
-	for (const TTuple<UDebugSkelMeshComponent*, FName>& ComponentTupple : SuspendedComponentArray)
-	{
-		ComponentTupple.Key->SetSkinWeightProfile(ComponentTupple.Value);
-	}
-	SuspendedComponentArray.Empty();
-}
-
-
 bool FSkeletalMeshEditor::HandleReimportMeshInternal(int32 SourceFileIndex /*= INDEX_NONE*/, bool bWithNewFile /*= false*/)
 {
-	FScopedSkeletalMeshPostEditChange ScopedPostEditChange(SkeletalMesh);
 	FScopedSuspendAlternateSkinWeightPreview ScopedSuspendAlternateSkinnWeightPreview(SkeletalMesh);
-
-	// Reimport the asset
-	const bool bResult = FReimportManager::Instance()->Reimport(SkeletalMesh, true, true, TEXT(""), nullptr, SourceFileIndex, bWithNewFile);
-
-	// Refresh skeleton tree
-	SkeletonTree->Refresh();
-
+	bool bResult = false;
+	{
+		FScopedSkeletalMeshPostEditChange ScopedPostEditChange(SkeletalMesh);
+		// Reimport the asset
+		bResult = FReimportManager::Instance()->Reimport(SkeletalMesh, true, true, TEXT(""), nullptr, SourceFileIndex, bWithNewFile);
+		// Refresh skeleton tree
+		SkeletonTree->Refresh();
+	}
 	return bResult;
 }
 
@@ -1132,51 +1079,53 @@ void FSkeletalMeshEditor::HandleReimportMeshWithNewFile(int32 SourceFileIndex /*
 
 void ReimportAllCustomLODs(USkeletalMesh* SkeletalMesh, UDebugSkelMeshComponent* PreviewMeshComponent, bool bWithNewFile)
 {
-	FScopedSkeletalMeshPostEditChange ScopedPostEditChange(SkeletalMesh);
 	FScopedSuspendAlternateSkinWeightPreview ScopedSuspendAlternateSkinnWeightPreview(SkeletalMesh);
-
-	//Find the dependencies of the generated LOD
-	TArray<bool> Dependencies;
-	Dependencies.AddZeroed(SkeletalMesh->GetLODNum());
-	//Avoid making LOD 0 to true in the dependencies since everything that should be regenerate base on LOD 0 is already regenerate at this point.
-	//But we need to regenerate every generated LOD base on any re-import custom LOD
-	//Reimport all custom LODs
-	for (int32 LodIndex = 1; LodIndex < SkeletalMesh->GetLODNum(); ++LodIndex)
 	{
-		//Do not reimport LOD that was re-import with the base mesh
-		if (SkeletalMesh->GetLODInfo(LodIndex)->bImportWithBaseMesh)
-		{
-			continue;
-		}
-		if (SkeletalMesh->GetLODInfo(LodIndex)->bHasBeenSimplified == false)
+		FScopedSkeletalMeshPostEditChange ScopedPostEditChange(SkeletalMesh);
+
+		//Find the dependencies of the generated LOD
+		TArray<bool> Dependencies;
+		Dependencies.AddZeroed(SkeletalMesh->GetLODNum());
+		//Avoid making LOD 0 to true in the dependencies since everything that should be regenerate base on LOD 0 is already regenerate at this point.
+		//But we need to regenerate every generated LOD base on any re-import custom LOD
+		//Reimport all custom LODs
+		for (int32 LodIndex = 1; LodIndex < SkeletalMesh->GetLODNum(); ++LodIndex)
 		{
-			FString SourceFilenameBackup = SkeletalMesh->GetLODInfo(LodIndex)->SourceImportFilename;
-			if (bWithNewFile)
+			//Do not reimport LOD that was re-import with the base mesh
+			if (SkeletalMesh->GetLODInfo(LodIndex)->bImportWithBaseMesh)
 			{
-				SkeletalMesh->GetLODInfo(LodIndex)->SourceImportFilename.Empty();
+				continue;
 			}
-
-			if (!FbxMeshUtils::ImportMeshLODDialog(SkeletalMesh, LodIndex, false))
+			if (SkeletalMesh->GetLODInfo(LodIndex)->bHasBeenSimplified == false)
 			{
+				FString SourceFilenameBackup = SkeletalMesh->GetLODInfo(LodIndex)->SourceImportFilename;
 				if (bWithNewFile)
 				{
-					SkeletalMesh->GetLODInfo(LodIndex)->SourceImportFilename = SourceFilenameBackup;
+					SkeletalMesh->GetLODInfo(LodIndex)->SourceImportFilename.Empty();
+				}
+
+				if (!FbxMeshUtils::ImportMeshLODDialog(SkeletalMesh, LodIndex, false))
+				{
+					if (bWithNewFile)
+					{
+						SkeletalMesh->GetLODInfo(LodIndex)->SourceImportFilename = SourceFilenameBackup;
+					}
+				}
+				else
+				{
+					Dependencies[LodIndex] = true;
 				}
 			}
-			else
+			else if (Dependencies[SkeletalMesh->GetLODInfo(LodIndex)->ReductionSettings.BaseLOD])
 			{
+				//Regenerate the LOD
+				FSkeletalMeshUpdateContext UpdateContext;
+				UpdateContext.SkeletalMesh = SkeletalMesh;
+				UpdateContext.AssociatedComponents.Push(PreviewMeshComponent);
+				FLODUtilities::SimplifySkeletalMeshLOD(UpdateContext, LodIndex, false);
 				Dependencies[LodIndex] = true;
 			}
 		}
-		else if(Dependencies[SkeletalMesh->GetLODInfo(LodIndex)->ReductionSettings.BaseLOD])
-		{
-			//Regenerate the LOD
-			FSkeletalMeshUpdateContext UpdateContext;
-			UpdateContext.SkeletalMesh = SkeletalMesh;
-			UpdateContext.AssociatedComponents.Push(PreviewMeshComponent);
-			FLODUtilities::SimplifySkeletalMeshLOD(UpdateContext, LodIndex, false);
-			Dependencies[LodIndex] = true;
-		}
 	}
 }
 
diff --git a/Engine/Source/Editor/UnrealEd/Classes/Animation/DebugSkelMeshComponent.h b/Engine/Source/Editor/UnrealEd/Classes/Animation/DebugSkelMeshComponent.h
--- a/Engine/Source/Editor/UnrealEd/Classes/Animation/DebugSkelMeshComponent.h
+++ b/Engine/Source/Editor/UnrealEd/Classes/Animation/DebugSkelMeshComponent.h
@@ -520,3 +520,28 @@ public:
 	}
 
 };
+
+
+/*
+ * This class is use to remove the alternate skinning preview from the multiple editor that can show it.
+ * Important it should be destroy after the PostEditChange of the skeletalmesh is done and the renderdata have been recreate
+ * i.e. FScopedSkeletalMeshPostEditChange should be create after FScopedSuspendAlternateSkinWeightPreview and delete before FScopedSuspendAlternateSkinWeightPreview
+ */
+class UNREALED_API FScopedSuspendAlternateSkinWeightPreview
+{
+public:
+	/*
+	 * This constructor suspend the alternate skinning preview for all editor component that use the specified skeletalmesh
+	 * Parameters:
+	 * @param InSkeletalMesh - SkeletalMesh use to know which preview component we have to suspend the alternate skinning preview.
+	 */
+	FScopedSuspendAlternateSkinWeightPreview(class USkeletalMesh* InSkeletalMesh);
+
+	/*
+	 * This destructor put back the preview alternate skinning
+	 */
+	~FScopedSuspendAlternateSkinWeightPreview();
+
+private:
+	TArray< TTuple<UDebugSkelMeshComponent*, FName> > SuspendedComponentArray;
+};
\ No newline at end of file
diff --git a/Engine/Source/Editor/UnrealEd/Private/Animation/DebugSkelMeshComponent.cpp b/Engine/Source/Editor/UnrealEd/Private/Animation/DebugSkelMeshComponent.cpp
--- a/Engine/Source/Editor/UnrealEd/Private/Animation/DebugSkelMeshComponent.cpp
+++ b/Engine/Source/Editor/UnrealEd/Private/Animation/DebugSkelMeshComponent.cpp
@@ -1057,3 +1057,37 @@ FDebugSkelMeshDynamicData::FDebugSkelMeshDynamicData(UDebugSkelMeshComponent* In
 	}
 }
 
+FScopedSuspendAlternateSkinWeightPreview::FScopedSuspendAlternateSkinWeightPreview(USkeletalMesh* SkeletalMesh)
+{
+	SuspendedComponentArray.Empty(2);
+	if (SkeletalMesh != nullptr)
+	{
+		// Now iterate over all skeletal mesh components and unregister them from the world, we will reregister them in the destructor
+		for (TObjectIterator<UDebugSkelMeshComponent> It; It; ++It)
+		{
+			UDebugSkelMeshComponent* DebugSKComp = *It;
+			if (DebugSKComp->SkeletalMesh == SkeletalMesh)
+			{
+				const FName ProfileName = DebugSKComp->GetCurrentSkinWeightProfileName();
+				if (ProfileName != NAME_None)
+				{
+					DebugSKComp->ClearSkinWeightProfile();
+					TTuple<UDebugSkelMeshComponent*, FName> ComponentTupple;
+					ComponentTupple.Key = DebugSKComp;
+					ComponentTupple.Value = ProfileName;
+					SuspendedComponentArray.Add(ComponentTupple);
+				}
+			}
+		}
+	}
+}
+
+FScopedSuspendAlternateSkinWeightPreview::~FScopedSuspendAlternateSkinWeightPreview()
+{
+	//Put back the skin weight profile for all editor debug component
+	for (const TTuple<UDebugSkelMeshComponent*, FName>& ComponentTupple : SuspendedComponentArray)
+	{
+		ComponentTupple.Key->SetSkinWeightProfile(ComponentTupple.Value);
+	}
+	SuspendedComponentArray.Empty();
+}
