SkeletalMesh toolbar conversion
#rnx

diff --git a/Engine/Source/Developer/MeshMergeUtilities/MeshMergeUtilities.Build.cs b/Engine/Source/Developer/MeshMergeUtilities/MeshMergeUtilities.Build.cs
--- a/Engine/Source/Developer/MeshMergeUtilities/MeshMergeUtilities.Build.cs
+++ b/Engine/Source/Developer/MeshMergeUtilities/MeshMergeUtilities.Build.cs
@@ -24,6 +24,7 @@ public class MeshMergeUtilities : ModuleRules
                 "SkeletalMeshEditor",
                 "MaterialBaking",
                 "MeshUtilitiesCommon",
+				"ToolMenus",
             }
 		);
         
diff --git a/Engine/Source/Developer/MeshMergeUtilities/Private/MeshMergeEditorExtensions.cpp b/Engine/Source/Developer/MeshMergeUtilities/Private/MeshMergeEditorExtensions.cpp
--- a/Engine/Source/Developer/MeshMergeUtilities/Private/MeshMergeEditorExtensions.cpp
+++ b/Engine/Source/Developer/MeshMergeUtilities/Private/MeshMergeEditorExtensions.cpp
@@ -17,10 +17,11 @@
 #include "MeshMergeUtilities.h"
 #include "Modules/ModuleManager.h"
 #include "MeshMergeModule.h"
+#include "ToolMenus.h"
+#include "SkeletalMeshToolMenuContext.h"
 
 #define LOCTEXT_NAMESPACE "MeshMergeEditorExtensions"
 
-FDelegateHandle FMeshMergeEditorExtensions::SkeletalMeshEditorExtenderHandle;
 FDelegateHandle FMeshMergeEditorExtensions::StaticMeshEditorExtenderHandle;
 
 void FMeshMergeEditorExtensions::OnModulesChanged(FName InModuleName, EModuleChangeReason InChangeReason)
@@ -28,11 +29,7 @@ void FMeshMergeEditorExtensions::OnModulesChanged(FName InModuleName, EModuleCha
 	// If one of the modules we are interested in is loaded apply editor extensions
 	if (InChangeReason == EModuleChangeReason::ModuleLoaded)
 	{		
-		if (InModuleName == "SkeletalMeshEditor")
-		{
-			AddSkeletalMeshEditorToolbarExtender();
-		}
-		else if (InModuleName == "StaticMeshEditor")
+		if (InModuleName == "StaticMeshEditor")
 		{
 			AddStaticMeshEditorToolbarExtender();
 		}
@@ -41,7 +38,6 @@ void FMeshMergeEditorExtensions::OnModulesChanged(FName InModuleName, EModuleCha
 
 void FMeshMergeEditorExtensions::RemoveExtenders()
 {
-	RemoveSkeletalMeshEditorToolbarExtender();
 	RemoveStaticMeshEditorToolbarExtender();
 }
 
@@ -77,7 +73,7 @@ void FMeshMergeEditorExtensions::RemoveStaticMeshEditorToolbarExtender()
 
 	if (StaticMeshEditorModule)
 	{
-		StaticMeshEditorModule->GetToolBarExtensibilityManager()->GetExtenderDelegates().RemoveAll([=](const auto& In) { return In.GetHandle() == SkeletalMeshEditorExtenderHandle; });
+		StaticMeshEditorModule->GetToolBarExtensibilityManager()->GetExtenderDelegates().RemoveAll([=](const auto& In) { return In.GetHandle() == StaticMeshEditorExtenderHandle; });
 	}
 }
 
@@ -96,57 +92,30 @@ void FMeshMergeEditorExtensions::HandleAddStaticMeshActionExtenderToToolbar(FToo
 	);
 }
 
-void FMeshMergeEditorExtensions::AddSkeletalMeshEditorToolbarExtender()
-{
-	ISkeletalMeshEditorModule& SkeletalMeshEditorModule = FModuleManager::Get().LoadModuleChecked<ISkeletalMeshEditorModule>("SkeletalMeshEditor");
-	auto& ToolbarExtenders = SkeletalMeshEditorModule.GetAllSkeletalMeshEditorToolbarExtenders();
-
-	ToolbarExtenders.Add(ISkeletalMeshEditorModule::FSkeletalMeshEditorToolbarExtender::CreateStatic(&FMeshMergeEditorExtensions::GetSkeletalMeshEditorToolbarExtender));
-	SkeletalMeshEditorExtenderHandle = ToolbarExtenders.Last().GetHandle();
-}
-
-void FMeshMergeEditorExtensions::RemoveSkeletalMeshEditorToolbarExtender()
+void FMeshMergeEditorExtensions::RegisterMenus()
 {
-	ISkeletalMeshEditorModule* SkeletalMeshEditorModule = FModuleManager::Get().GetModulePtr<ISkeletalMeshEditorModule>("SkeletalMeshEditor");
-	if (SkeletalMeshEditorModule)
 	{
-		typedef ISkeletalMeshEditorModule::FSkeletalMeshEditorToolbarExtender DelegateType;
-		SkeletalMeshEditorModule->GetAllSkeletalMeshEditorToolbarExtenders().RemoveAll([=](const DelegateType& In) { return In.GetHandle() == SkeletalMeshEditorExtenderHandle; });
-	}
-}
-
-TSharedRef<FExtender> FMeshMergeEditorExtensions::GetSkeletalMeshEditorToolbarExtender(const TSharedRef<FUICommandList> CommandList, TSharedRef<ISkeletalMeshEditor> InSkeletalMeshEditor)
-{
-	TSharedRef<FExtender> Extender = MakeShareable(new FExtender);
-	UMeshComponent* MeshComponent = InSkeletalMeshEditor->GetPersonaToolkit()->GetPreviewMeshComponent();
-
-	// Add button on skeletal mesh editor toolbar
-	Extender->AddToolBarExtension(
-		"Asset",
-		EExtensionHook::After,
-		CommandList,
-		FToolBarExtensionDelegate::CreateStatic(&FMeshMergeEditorExtensions::HandleAddSkeletalMeshActionExtenderToToolbar, MeshComponent)
-	);
-
-	return Extender;
-}
-
-void FMeshMergeEditorExtensions::HandleAddSkeletalMeshActionExtenderToToolbar(FToolBarBuilder& ParentToolbarBuilder, UMeshComponent* InMeshComponent)
-{
-	ParentToolbarBuilder.AddToolBarButton(
-		FUIAction(FExecuteAction::CreateLambda([InMeshComponent]()
-		{
-			if (USkeletalMeshComponent* SkeletalMeshComponent = Cast<USkeletalMeshComponent>(InMeshComponent))
+		UToolMenu* Toolbar = UToolMenus::Get()->ExtendMenu("AssetEditor.SkeletalMeshEditor.ToolBar");
+		FToolMenuSection& Section = Toolbar->FindOrAddSection("SkeletalMesh");
+		Section.AddEntry(FToolMenuEntry::InitToolBarButton(
+			"BakeMaterials",
+			FToolMenuExecuteAction::CreateLambda([](const FToolMenuContext& InMenuContext)
 			{
-				IMeshMergeModule& Module = FModuleManager::Get().LoadModuleChecked<IMeshMergeModule>("MeshMergeUtilities");				
-				Module.GetUtilities().BakeMaterialsForComponent(SkeletalMeshComponent);
-			}
-		})),
-		NAME_None,
-		LOCTEXT("BakeMaterials", "Bake out Materials"),
-		LOCTEXT("BakeMaterialsTooltip", "Bake out Materials for given LOD(s)."),
-		FSlateIcon("EditorStyle", "Persona.BakeMaterials")
-	);
+				USkeletalMeshToolMenuContext* Context = InMenuContext.FindContext<USkeletalMeshToolMenuContext>();
+				if (Context && Context->SkeletalMeshEditor.IsValid())
+				{
+					if (UDebugSkelMeshComponent* SkeletalMeshComponent = Context->SkeletalMeshEditor.Pin()->GetPersonaToolkit()->GetPreviewMeshComponent())
+					{
+						IMeshMergeModule& Module = FModuleManager::Get().LoadModuleChecked<IMeshMergeModule>("MeshMergeUtilities");
+						Module.GetUtilities().BakeMaterialsForComponent(SkeletalMeshComponent);
+					}
+				}
+			}),
+			LOCTEXT("BakeMaterials", "Bake out Materials"),
+			LOCTEXT("BakeMaterialsTooltip", "Bake out Materials for given LOD(s)."),
+			FSlateIcon("EditorStyle", "Persona.BakeMaterials")
+		));
+	}
 }
 
 #undef LOCTEXT_NAMESPACE // "MeshMergeEditorExtensions"
\ No newline at end of file
diff --git a/Engine/Source/Developer/MeshMergeUtilities/Private/MeshMergeEditorExtensions.h b/Engine/Source/Developer/MeshMergeUtilities/Private/MeshMergeEditorExtensions.h
--- a/Engine/Source/Developer/MeshMergeUtilities/Private/MeshMergeEditorExtensions.h
+++ b/Engine/Source/Developer/MeshMergeUtilities/Private/MeshMergeEditorExtensions.h
@@ -20,19 +20,13 @@ class FMeshMergeEditorExtensions
 public:
 	static void OnModulesChanged(FName InModuleName, EModuleChangeReason InChangeReason);
 	static void RemoveExtenders();
+	static void RegisterMenus();
 	
 	/** Callback functionality for static mesh editor extension */
 	static void AddStaticMeshEditorToolbarExtender();
 	static void RemoveStaticMeshEditorToolbarExtender();
 	static void HandleAddStaticMeshActionExtenderToToolbar(FToolBarBuilder& ParentToolbarBuilder, UStaticMesh* Mesh);
 	static TSharedRef<FExtender> GetStaticMeshEditorToolbarExtender(const TSharedRef<FUICommandList> CommandList, const TArray<UObject*> Objects);
-	
-	/** Callback functionality for skeletal mesh editor extension */
-	static void AddSkeletalMeshEditorToolbarExtender();
-	static void RemoveSkeletalMeshEditorToolbarExtender();
-	static TSharedRef<FExtender> GetSkeletalMeshEditorToolbarExtender(const TSharedRef<FUICommandList> CommandList, TSharedRef<ISkeletalMeshEditor> InSkeletalMeshEditor);
-	static void HandleAddSkeletalMeshActionExtenderToToolbar(FToolBarBuilder& ParentToolbarBuilder, UMeshComponent* InMeshComponent);
-	
-	static FDelegateHandle SkeletalMeshEditorExtenderHandle;
+
 	static FDelegateHandle StaticMeshEditorExtenderHandle;
 };
\ No newline at end of file
diff --git a/Engine/Source/Developer/MeshMergeUtilities/Private/MeshMergeModule.cpp b/Engine/Source/Developer/MeshMergeUtilities/Private/MeshMergeModule.cpp
--- a/Engine/Source/Developer/MeshMergeUtilities/Private/MeshMergeModule.cpp
+++ b/Engine/Source/Developer/MeshMergeUtilities/Private/MeshMergeModule.cpp
@@ -3,6 +3,8 @@
 #include "MeshMergeModule.h"
 #include "MeshMergeUtilities.h"
 #include "Modules/ModuleManager.h"
+#include "MeshMergeEditorExtensions.h"
+#include "ToolMenus.h"
 
 class FMeshMergeModule : public IMeshMergeModule
 {
@@ -16,8 +18,27 @@ public:
 	{
 		return Utilities;
 	}
+
+	virtual void StartupModule() override
+	{
+		UToolMenus::RegisterStartupCallback(FSimpleMulticastDelegate::FDelegate::CreateRaw(this, &FMeshMergeModule::RegisterMenus));
+	}
+
+	virtual void ShutdownModule() override
+	{
+		UToolMenus::UnRegisterStartupCallback(this);
+		UToolMenus::UnregisterOwner(this);
+	}
+
 protected:
 	FMeshMergeUtilities Utilities;
+
+private:
+	void RegisterMenus()
+	{
+		FToolMenuOwnerScoped OwnerScoped(this);
+		FMeshMergeEditorExtensions::RegisterMenus();
+	}
 };
 
 
diff --git a/Engine/Source/Developer/MeshUtilities/MeshUtilities.Build.cs b/Engine/Source/Developer/MeshUtilities/MeshUtilities.Build.cs
--- a/Engine/Source/Developer/MeshUtilities/MeshUtilities.Build.cs
+++ b/Engine/Source/Developer/MeshUtilities/MeshUtilities.Build.cs
@@ -37,7 +37,8 @@ public class MeshUtilities : ModuleRules
                 "MeshUtilitiesCommon",
                 "MeshDescription",
 				"StaticMeshDescription",
-                "MeshDescriptionOperations"
+                "MeshDescriptionOperations",
+				"ToolMenus",
             }
 		);
 
diff --git a/Engine/Source/Developer/MeshUtilities/Private/MeshUtilities.cpp b/Engine/Source/Developer/MeshUtilities/Private/MeshUtilities.cpp
--- a/Engine/Source/Developer/MeshUtilities/Private/MeshUtilities.cpp
+++ b/Engine/Source/Developer/MeshUtilities/Private/MeshUtilities.cpp
@@ -17,6 +17,8 @@
 #include "Framework/Commands/UICommandList.h"
 #include "Framework/MultiBox/MultiBoxExtender.h"
 #include "Framework/MultiBox/MultiBoxBuilder.h"
+#include "ToolMenus.h"
+#include "SkeletalMeshToolMenuContext.h"
 #include "Components/MeshComponent.h"
 #include "RawIndexBuffer.h"
 #include "Components/StaticMeshComponent.h"
@@ -5546,20 +5548,21 @@ void FMeshUtilities::StartupModule()
 			{
 				AddAnimationEditorToolbarExtender();
 			}
-			else if (InModuleName == "SkeletalMeshEditor")
-			{
-				AddSkeletalMeshEditorToolbarExtender();
-			}
 			else if (InModuleName == "SkeletonEditor")
 			{
 				AddSkeletonEditorToolbarExtender();
 			}
 		}
 	});
+
+	UToolMenus::RegisterStartupCallback(FSimpleMulticastDelegate::FDelegate::CreateRaw(this, &FMeshUtilities::RegisterMenus));
 }
 
 void FMeshUtilities::ShutdownModule()
 {
+	UToolMenus::UnRegisterStartupCallback(this);
+	UToolMenus::UnregisterOwner(this);
+
 	static const FName PropertyEditorModuleName("PropertyEditor");
 	if(FModuleManager::Get().IsModuleLoaded(PropertyEditorModuleName))
 	{
@@ -5573,12 +5576,42 @@ void FMeshUtilities::ShutdownModule()
 	RemoveLevelViewportMenuExtender();
 	RemoveAnimationBlueprintEditorToolbarExtender();
 	RemoveAnimationEditorToolbarExtender();
-	RemoveSkeletalMeshEditorToolbarExtender();
 	RemoveSkeletonEditorToolbarExtender();
 	FModuleManager::Get().OnModulesChanged().Remove(ModuleLoadedDelegateHandle);
 	VersionString.Empty();
 }
 
+void FMeshUtilities::RegisterMenus()
+{
+	FToolMenuOwnerScoped OwnerScoped(this);
+
+	static auto AddMakeStaticMeshToolbarButton = [this](FToolMenuSection& InSection, const FToolMenuExecuteAction& InAction)
+	{
+		InSection.AddEntry(FToolMenuEntry::InitToolBarButton(
+			"MakeStaticMesh",
+			InAction,
+			LOCTEXT("MakeStaticMesh", "Make Static Mesh"),
+			LOCTEXT("MakeStaticMeshTooltip", "Make a new static mesh out of the preview's current pose."),
+			FSlateIcon("EditorStyle", "Persona.ConvertToStaticMesh")
+		));
+	};
+
+	{
+		UToolMenu* Toolbar = UToolMenus::Get()->ExtendMenu("AssetEditor.SkeletalMeshEditor.ToolBar");
+		FToolMenuSection& Section = Toolbar->FindOrAddSection("SkeletalMesh");
+		AddMakeStaticMeshToolbarButton(Section, FToolMenuExecuteAction::CreateLambda([this](const FToolMenuContext& InMenuContext)
+		{
+			USkeletalMeshToolMenuContext* Context = InMenuContext.FindContext<USkeletalMeshToolMenuContext>();
+			if (Context && Context->SkeletalMeshEditor.IsValid())
+			{
+				if (UMeshComponent* MeshComponent = Context->SkeletalMeshEditor.Pin()->GetPersonaToolkit()->GetPreviewMeshComponent())
+				{
+					ConvertMeshesToStaticMesh(TArray<UMeshComponent*>({ MeshComponent }), MeshComponent->GetComponentToWorld());
+				}
+			}
+		}));
+	}
+}
 
 bool FMeshUtilities::GenerateUniqueUVsForSkeletalMesh(const FSkeletalMeshLODModel& LODModel, int32 TextureResolution, TArray<FVector2D>& OutTexCoords) const
 {
@@ -5838,25 +5871,6 @@ TSharedRef<FExtender> FMeshUtilities::GetAnimationEditorToolbarExtender(const TS
 	return Extender;
 }
 
-void FMeshUtilities::AddSkeletalMeshEditorToolbarExtender()
-{
-	ISkeletalMeshEditorModule& SkeletalMeshEditorModule = FModuleManager::Get().LoadModuleChecked<ISkeletalMeshEditorModule>("SkeletalMeshEditor");
-	auto& ToolbarExtenders = SkeletalMeshEditorModule.GetAllSkeletalMeshEditorToolbarExtenders();
-
-	ToolbarExtenders.Add(ISkeletalMeshEditorModule::FSkeletalMeshEditorToolbarExtender::CreateRaw(this, &FMeshUtilities::GetSkeletalMeshEditorToolbarExtender));
-	SkeletalMeshEditorExtenderHandle = ToolbarExtenders.Last().GetHandle();
-}
-
-void FMeshUtilities::RemoveSkeletalMeshEditorToolbarExtender()
-{
-	ISkeletalMeshEditorModule* SkeletalMeshEditorModule = FModuleManager::Get().GetModulePtr<ISkeletalMeshEditorModule>("SkeletalMeshEditor");
-	if (SkeletalMeshEditorModule)
-	{
-		typedef ISkeletalMeshEditorModule::FSkeletalMeshEditorToolbarExtender DelegateType;
-		SkeletalMeshEditorModule->GetAllSkeletalMeshEditorToolbarExtenders().RemoveAll([=](const DelegateType& In) { return In.GetHandle() == SkeletalMeshEditorExtenderHandle; });
-	}
-}
-
 TSharedRef<FExtender> FMeshUtilities::GetSkeletalMeshEditorToolbarExtender(const TSharedRef<FUICommandList> CommandList, TSharedRef<ISkeletalMeshEditor> InSkeletalMeshEditor)
 {
 	TSharedRef<FExtender> Extender = MakeShareable(new FExtender);
diff --git a/Engine/Source/Developer/MeshUtilities/Private/MeshUtilitiesPrivate.h b/Engine/Source/Developer/MeshUtilities/Private/MeshUtilitiesPrivate.h
--- a/Engine/Source/Developer/MeshUtilities/Private/MeshUtilitiesPrivate.h
+++ b/Engine/Source/Developer/MeshUtilities/Private/MeshUtilitiesPrivate.h
@@ -166,6 +166,8 @@ private:
 
 	virtual void GenerateRuntimeSkinWeightData(const FSkeletalMeshLODModel* ImportedModel, const TArray<FRawSkinWeight>& InRawSkinWeights, FRuntimeSkinWeightProfileData& InOutSkinWeightOverrideData) const override;
 
+	void RegisterMenus();
+
 	// Need to call some members from this class, (which is internal to this module)
 	friend class FStaticMeshUtilityBuilder;
 
@@ -182,10 +184,6 @@ protected:
 
 	TSharedRef<FExtender> GetAnimationEditorToolbarExtender(const TSharedRef<FUICommandList> CommandList, TSharedRef<IAnimationEditor> InAnimationEditor);
 
-	void AddSkeletalMeshEditorToolbarExtender();
-
-	void RemoveSkeletalMeshEditorToolbarExtender();
-
 	TSharedRef<FExtender> GetSkeletalMeshEditorToolbarExtender(const TSharedRef<FUICommandList> CommandList, TSharedRef<ISkeletalMeshEditor> InSkeletalMeshEditor);
 
 	void AddSkeletonEditorToolbarExtender();
@@ -208,7 +206,6 @@ protected:
 	FDelegateHandle LevelViewportExtenderHandle;
 	FDelegateHandle AnimationBlueprintEditorExtenderHandle;
 	FDelegateHandle AnimationEditorExtenderHandle;
-	FDelegateHandle SkeletalMeshEditorExtenderHandle;
 	FDelegateHandle SkeletonEditorExtenderHandle;
 };
 
diff --git a/Engine/Source/Editor/ClothPainter/ClothPainter.Build.cs b/Engine/Source/Editor/ClothPainter/ClothPainter.Build.cs
--- a/Engine/Source/Editor/ClothPainter/ClothPainter.Build.cs
+++ b/Engine/Source/Editor/ClothPainter/ClothPainter.Build.cs
@@ -31,7 +31,8 @@ public class ClothPainter : ModuleRules
                 "ClothingSystemRuntimeInterface",
                 "ClothingSystemEditorInterface",
                 "SkeletalMeshEditor",
-                "AdvancedPreviewScene"
+                "AdvancedPreviewScene",
+				"ToolMenus",
             }
 		);
     }
diff --git a/Engine/Source/Editor/ClothPainter/Private/ClothPaintingModule.cpp b/Engine/Source/Editor/ClothPainter/Private/ClothPaintingModule.cpp
--- a/Engine/Source/Editor/ClothPainter/Private/ClothPaintingModule.cpp
+++ b/Engine/Source/Editor/ClothPainter/Private/ClothPaintingModule.cpp
@@ -18,6 +18,8 @@
 #include "ClothPaintToolCommands.h"
 #include "ISkeletalMeshEditorModule.h"
 #include "Framework/MultiBox/MultiBoxBuilder.h"
+#include "ToolMenus.h"
+#include "SkeletalMeshToolMenuContext.h"
 #include "ClothPainterCommands.h"
 #include "Widgets/Docking/SDockTab.h"
 
@@ -68,6 +70,8 @@ void FClothPaintingModule::StartupModule()
 	ClothPaintToolCommands::RegisterClothPaintToolCommands();
 	FClothPainterCommands::Register();
 
+	UToolMenus::RegisterStartupCallback(FSimpleMulticastDelegate::FDelegate::CreateRaw(this, &FClothPaintingModule::RegisterMenus));
+
 	if(!SkelMeshEditorExtenderHandle.IsValid())
 	{
 		ISkeletalMeshEditorModule& SkelMeshEditorModule = FModuleManager::Get().LoadModuleChecked<ISkeletalMeshEditorModule>("SkeletalMeshEditor");
@@ -112,21 +116,32 @@ TSharedRef<FExtender> FClothPaintingModule::ExtendSkelMeshEditorToolbar(const TS
 		FIsActionChecked::CreateRaw(this, &FClothPaintingModule::GetIsPaintToolsButtonChecked, Ptr)
 	);
 
-	ToolbarExtender->AddToolBarExtension("Asset", EExtensionHook::After, InCommandList, FToolBarExtensionDelegate::CreateLambda(
-		[this, Ptr](FToolBarBuilder& Builder)
-	{
-		Builder.AddToolBarButton(
-			FClothPainterCommands::Get().TogglePaintMode, 
-			NAME_None, 
-			TAttribute<FText>::Create(TAttribute<FText>::FGetter::CreateRaw(this, &FClothPaintingModule::GetPaintToolsButtonText, Ptr)), 
-			TAttribute<FText>(), 
-			FSlateIcon(FEditorStyle::GetStyleSetName(), 
-			"LevelEditor.MeshPaintMode.TexturePaint"));
-	}));
-
 	return ToolbarExtender.ToSharedRef();
 }
 
+void FClothPaintingModule::RegisterMenus()
+{
+	FToolMenuOwnerScoped OwnerScoped(this);
+
+	UToolMenu* Toolbar = UToolMenus::Get()->ExtendMenu("AssetEditor.SkeletalMeshEditor.ToolBar");
+	{
+		FToolMenuSection& Section = Toolbar->FindOrAddSection("SkeletalMesh");
+		Section.AddDynamicEntry("TogglePaintMode", FNewToolMenuSectionDelegate::CreateLambda([this](FToolMenuSection& InSection)
+		{
+			USkeletalMeshToolMenuContext* Context = InSection.FindContext<USkeletalMeshToolMenuContext>();
+			if (Context && Context->SkeletalMeshEditor.IsValid())
+			{
+				InSection.AddEntry(FToolMenuEntry::InitToolBarButton(
+					FClothPainterCommands::Get().TogglePaintMode,
+					TAttribute<FText>::Create(TAttribute<FText>::FGetter::CreateRaw(this, &FClothPaintingModule::GetPaintToolsButtonText, Context->SkeletalMeshEditor)),
+					FText(),
+					FSlateIcon(FEditorStyle::GetStyleSetName(), "LevelEditor.MeshPaintMode.TexturePaint")
+				));	
+			}
+		}));
+	}
+}
+
 FText FClothPaintingModule::GetPaintToolsButtonText(TWeakPtr<ISkeletalMeshEditor> InSkeletalMeshEditor) const
 {
 	TSharedPtr<SClothPaintTab> ClothTab = GetActiveClothTab(InSkeletalMeshEditor, false);
@@ -194,6 +209,9 @@ void FClothPaintingModule::ShutdownModule()
 {
 	ShutdownMode();
 
+	UToolMenus::UnRegisterStartupCallback(this);
+	UToolMenus::UnregisterOwner(this);
+
 	// Remove skel mesh editor extenders
 	ISkeletalMeshEditorModule* SkelMeshEditorModule = FModuleManager::GetModulePtr<ISkeletalMeshEditorModule>("SkeletalMeshEditor");
 	if (SkelMeshEditorModule)
diff --git a/Engine/Source/Editor/ClothPainter/Public/ClothPaintingModule.h b/Engine/Source/Editor/ClothPainter/Public/ClothPaintingModule.h
--- a/Engine/Source/Editor/ClothPainter/Public/ClothPaintingModule.h
+++ b/Engine/Source/Editor/ClothPainter/Public/ClothPaintingModule.h
@@ -36,6 +36,9 @@ protected:
 
 private:
 
+	// Extend toolbars and menus
+	void RegisterMenus();
+
 	// Extends a skeletal mesh editor instance toolbar
 	TSharedRef<FExtender> ExtendSkelMeshEditorToolbar(const TSharedRef<FUICommandList> InCommandList, TSharedRef<ISkeletalMeshEditor> InSkeletalMeshEditor);
 
diff --git a/Engine/Source/Editor/Kismet/Private/BlueprintEditor.cpp b/Engine/Source/Editor/Kismet/Private/BlueprintEditor.cpp
--- a/Engine/Source/Editor/Kismet/Private/BlueprintEditor.cpp
+++ b/Engine/Source/Editor/Kismet/Private/BlueprintEditor.cpp
@@ -1906,18 +1906,6 @@ void FBlueprintEditor::InitToolMenuContext(FToolMenuContext& MenuContext)
 	MenuContext.AddObject(Context);
 }
 
-UToolMenu* FBlueprintEditor::RegisterModeToolbarIfUnregistered(const FName InModeName)
-{
-	FName ParentToolbarName;
-	const FName ModeSpecificToolBarName = GetToolMenuToolbarNameForMode(InModeName, ParentToolbarName);
-	if (!UToolMenus::Get()->IsMenuRegistered(ModeSpecificToolBarName))
-	{
-		return UToolMenus::Get()->RegisterMenu(ModeSpecificToolBarName, ParentToolbarName, EMultiBoxType::ToolBar);
-	}
-
-	return nullptr;
-}
-
 void FBlueprintEditor::InitalizeExtenders()
 {
 	TSharedPtr<FExtender> MenuExtender = MakeShareable(new FExtender);
diff --git a/Engine/Source/Editor/Kismet/Private/WorkflowCentricApplication.cpp b/Engine/Source/Editor/Kismet/Private/WorkflowCentricApplication.cpp
--- a/Engine/Source/Editor/Kismet/Private/WorkflowCentricApplication.cpp
+++ b/Engine/Source/Editor/Kismet/Private/WorkflowCentricApplication.cpp
@@ -5,6 +5,7 @@
 #include "WorkflowOrientedApp/WorkflowTabFactory.h"
 #include "WorkflowOrientedApp/WorkflowTabManager.h"
 #include "WorkflowOrientedApp/ApplicationMode.h"
+#include "ToolMenus.h"
 
 #define LOCTEXT_NAMESPACE "WorkflowCentricApplication"
 
@@ -43,6 +44,18 @@ FName FWorkflowCentricApplication::GetToolMenuToolbarNameForMode(const FName InM
 	return BaseMenuName;
 }
 
+UToolMenu* FWorkflowCentricApplication::RegisterModeToolbarIfUnregistered(const FName InModeName)
+{
+	FName ParentToolbarName;
+	const FName ModeSpecificToolbarName = GetToolMenuToolbarNameForMode(InModeName, ParentToolbarName);
+	if (!UToolMenus::Get()->IsMenuRegistered(ModeSpecificToolbarName))
+	{
+		return UToolMenus::Get()->RegisterMenu(ModeSpecificToolbarName, ParentToolbarName, EMultiBoxType::ToolBar);
+	}
+
+	return nullptr;
+}
+
 FName FWorkflowCentricApplication::GetCurrentMode() const
 {
 	return CurrentAppModePtr.IsValid() ? CurrentAppModePtr->GetModeName() : NAME_None;
diff --git a/Engine/Source/Editor/Kismet/Public/BlueprintEditor.h b/Engine/Source/Editor/Kismet/Public/BlueprintEditor.h
--- a/Engine/Source/Editor/Kismet/Public/BlueprintEditor.h
+++ b/Engine/Source/Editor/Kismet/Public/BlueprintEditor.h
@@ -242,14 +242,6 @@ public:
 	/** Add context objects for menus and toolbars */
 	virtual void InitToolMenuContext(FToolMenuContext& MenuContext) override;
 
-	/**
-	* Registers toolbar for a mode if not already registered
-	*
-	* @param InModeName - Name of the mode to register toolbar for
-	* @return nullptr if toolbar already registered.
-	*/
-	UToolMenu* RegisterModeToolbarIfUnregistered(const FName InModeName);
-
 	/** Check the Node Title is visible */
 	bool IsNodeTitleVisible(const UEdGraphNode* Node, bool bRequestRename);
 
diff --git a/Engine/Source/Editor/Kismet/Public/WorkflowOrientedApp/WorkflowCentricApplication.h b/Engine/Source/Editor/Kismet/Public/WorkflowOrientedApp/WorkflowCentricApplication.h
--- a/Engine/Source/Editor/Kismet/Public/WorkflowOrientedApp/WorkflowCentricApplication.h
+++ b/Engine/Source/Editor/Kismet/Public/WorkflowOrientedApp/WorkflowCentricApplication.h
@@ -8,6 +8,7 @@
 class FApplicationMode;
 class FWorkflowAllowedTabSet;
 class FWorkflowTabFactory;
+class UToolMenu;
 
 // Delegate for mutating a mode
 DECLARE_DELEGATE_RetVal_TwoParams(TSharedRef<FApplicationMode>, FWorkflowApplicationModeExtender, const FName /*ModeName*/, TSharedRef<FApplicationMode> /*InMode*/);
@@ -43,6 +44,14 @@ public:
 
 	// Gets the mode extender list for all workflow applications (append to customize a specific mode)
 	static TArray<FWorkflowApplicationModeExtender>& GetModeExtenderList() { return ModeExtenderList; }
+
+	/**
+	* Registers toolbar for a mode if not already registered
+	*
+	* @param InModeName - Name of the mode to register toolbar for
+	* @return nullptr if toolbar already registered.
+	*/
+	virtual UToolMenu* RegisterModeToolbarIfUnregistered(const FName InModeName);
 protected:
 
 	virtual void AddApplicationMode(FName ModeName, TSharedRef<FApplicationMode> Mode);
diff --git a/Engine/Source/Editor/SkeletalMeshEditor/Private/SkeletalMeshEditor.cpp b/Engine/Source/Editor/SkeletalMeshEditor/Private/SkeletalMeshEditor.cpp
--- a/Engine/Source/Editor/SkeletalMeshEditor/Private/SkeletalMeshEditor.cpp
+++ b/Engine/Source/Editor/SkeletalMeshEditor/Private/SkeletalMeshEditor.cpp
@@ -31,6 +31,8 @@
 #include "Widgets/Layout/SBox.h"
 #include "Widgets/Layout/SBorder.h"
 #include "Framework/Application/SlateApplication.h"
+#include "ToolMenus.h"
+#include "SkeletalMeshToolMenuContext.h"
 #include "EditorViewportClient.h"
 #include "Settings/EditorExperimentalSettings.h"
 #include "Algo/Transform.h"
@@ -220,13 +222,14 @@ void FSkeletalMeshEditor::InitSkeletalMeshEditor(const EToolkitMode::Type Mode,
 	ISkeletonEditorModule& SkeletonEditorModule = FModuleManager::GetModuleChecked<ISkeletonEditorModule>("SkeletonEditor");
 	SkeletonTree = SkeletonEditorModule.CreateSkeletonTree(PersonaToolkit->GetSkeleton(), SkeletonTreeArgs);
 
+	BindCommands();
+	RegisterToolbar();
+
 	const bool bCreateDefaultStandaloneMenu = true;
 	const bool bCreateDefaultToolbar = true;
 	const TSharedRef<FTabManager::FLayout> DummyLayout = FTabManager::NewLayout("NullLayout")->AddArea(FTabManager::NewPrimaryArea());
 	FAssetEditorToolkit::InitAssetEditor(Mode, InitToolkitHost, SkeletalMeshEditorAppIdentifier, DummyLayout, bCreateDefaultStandaloneMenu, bCreateDefaultToolbar, InSkeletalMesh);
 
-	BindCommands();
-
 	AddApplicationMode(
 		SkeletalMeshEditorModes::SkeletalMeshEditorMode,
 		MakeShareable(new FSkeletalMeshEditorMode(SharedThis(this), SkeletonTree.ToSharedRef())));
@@ -286,95 +289,110 @@ void FSkeletalMeshEditor::BindCommands()
 		FExecuteAction::CreateRaw(&GetPersonaToolkit()->GetPreviewScene().Get(), &IPersonaPreviewScene::TogglePlayback));
 }
 
-void FSkeletalMeshEditor::ExtendToolbar()
+TSharedPtr<FSkeletalMeshEditor> FSkeletalMeshEditor::GetSkeletalMeshEditor(const FToolMenuContext& InMenuContext)
 {
-	// If the ToolbarExtender is valid, remove it before rebuilding it
-	if (ToolbarExtender.IsValid())
+	if (USkeletalMeshToolMenuContext* Context = InMenuContext.FindContext<USkeletalMeshToolMenuContext>())
 	{
-		RemoveToolbarExtender(ToolbarExtender);
-		ToolbarExtender.Reset();
+		if (Context->SkeletalMeshEditor.IsValid())
+		{
+			return StaticCastSharedPtr<FSkeletalMeshEditor>(Context->SkeletalMeshEditor.Pin());
+		}
 	}
 
-	ToolbarExtender = MakeShareable(new FExtender);
+	return TSharedPtr<FSkeletalMeshEditor>();
+}
 
-	auto ConstructReimportContextMenu = [this]()
+void FSkeletalMeshEditor::RegisterReimportContextMenu(const FName InBaseMenuName)
+{
+	static auto ReimportMeshWithNewFileAction = [](const FToolMenuContext& InMenuContext, int32 SourceFileIndex)
 	{
-		bool bShowSubMenu = SkeletalMesh != nullptr && SkeletalMesh->AssetImportData != nullptr && SkeletalMesh->AssetImportData->GetSourceFileCount() > 1;
-		FMenuBuilder MenuBuilder(true, nullptr);
-		
-		if (!bShowSubMenu)
+		TSharedPtr<FSkeletalMeshEditor> SkeletalMeshEditor = GetSkeletalMeshEditor(InMenuContext);
+		if (SkeletalMeshEditor.IsValid())
 		{
-			//Reimport
-			MenuBuilder.AddMenuEntry(FSkeletalMeshEditorCommands::Get().ReimportMesh->GetLabel(),
-				FSkeletalMeshEditorCommands::Get().ReimportMesh->GetDescription(),
-				FSlateIcon(),
-				FUIAction(FExecuteAction::CreateSP(this, &FSkeletalMeshEditor::HandleReimportMesh, 0)));
-			
-			MenuBuilder.AddMenuEntry(FSkeletalMeshEditorCommands::Get().ReimportMeshWithNewFile->GetLabel(),
-				FSkeletalMeshEditorCommands::Get().ReimportMeshWithNewFile->GetDescription(),
-				FSlateIcon(),
-				FUIAction(FExecuteAction::CreateSP(this, &FSkeletalMeshEditor::HandleReimportMeshWithNewFile, 0)));
-
-			//Reimport ALL
-			MenuBuilder.AddMenuEntry(FSkeletalMeshEditorCommands::Get().ReimportAllMesh->GetLabel(),
-				FSkeletalMeshEditorCommands::Get().ReimportAllMesh->GetDescription(),
-				FSlateIcon(),
-				FUIAction(FExecuteAction::CreateSP(this, &FSkeletalMeshEditor::HandleReimportAllMesh, 0)));
-
-			MenuBuilder.AddMenuEntry(FSkeletalMeshEditorCommands::Get().ReimportAllMeshWithNewFile->GetLabel(),
-				FSkeletalMeshEditorCommands::Get().ReimportAllMeshWithNewFile->GetDescription(),
-				FSlateIcon(),
-				FUIAction(FExecuteAction::CreateSP(this, &FSkeletalMeshEditor::HandleReimportAllMeshWithNewFile, 0)));
-
-			FText ReimportMultiSources = LOCTEXT("ReimportMultiSources", "Reimport Content");
-			FText ReimportMultiSourcesTooltip = LOCTEXT("ReimportMultiSourcesTooltip", "Reimport Geometry or Skinning Weights content, this will create multi import source file.");
-
-			auto CreateMultiContentSubMenu = [this](FMenuBuilder& SubMenuBuilder)
+			if (USkeletalMesh* FoundSkeletalMesh = SkeletalMeshEditor->SkeletalMesh)
 			{
-				auto UpdateImportDataContentType = [this](int32 SourceFileIndex)
+				UFbxSkeletalMeshImportData* SkeletalMeshImportData = Cast<UFbxSkeletalMeshImportData>(FoundSkeletalMesh->AssetImportData);
+				if (SkeletalMeshImportData)
 				{
-					UFbxSkeletalMeshImportData* SkeletalMeshImportData = Cast<UFbxSkeletalMeshImportData>(SkeletalMesh->AssetImportData);
-					if (SkeletalMeshImportData)
+					SkeletalMeshImportData->ImportContentType = SourceFileIndex == 0 ? EFBXImportContentType::FBXICT_All : SourceFileIndex == 1 ? EFBXImportContentType::FBXICT_Geometry : EFBXImportContentType::FBXICT_SkinningWeights;
+					SkeletalMeshEditor->HandleReimportMeshWithNewFile(SourceFileIndex);
+				}
+			}
+		}
+	};
+
+	static auto CreateMultiContentSubMenu = [](UToolMenu* InMenu)
+	{
+		FToolMenuSection& Section = InMenu->AddSection("Reimport");
+
+		TSharedPtr<FSkeletalMeshEditor> SkeletalMeshEditor = GetSkeletalMeshEditor(InMenu->Context);
+		if (SkeletalMeshEditor.IsValid())
+		{
+			Section.AddMenuEntry(
+				"ReimportGeometryContentLabel",
+				LOCTEXT("ReimportGeometryContentLabel", "Geometry"),
+				LOCTEXT("ReimportGeometryContentLabelTooltipTooltip", "Reimport Geometry Only"),
+				FSlateIcon(FEditorStyle::GetStyleSetName(), "ContentBrowser.AssetActions.ReimportAsset"),
+				FToolMenuExecuteAction::CreateLambda(ReimportMeshWithNewFileAction, 1)
+			);
+
+			Section.AddMenuEntry(
+				"ReimportSkinningAndWeightsContentLabel",
+				LOCTEXT("ReimportSkinningAndWeightsContentLabel", "Skinning And Weights"),
+				LOCTEXT("ReimportSkinningAndWeightsContentLabelTooltipTooltip", "Reimport Skinning And Weights Only"),
+				FSlateIcon(FEditorStyle::GetStyleSetName(), "ContentBrowser.AssetActions.ReimportAsset"),
+				FToolMenuExecuteAction::CreateLambda(ReimportMeshWithNewFileAction, 2)
+			);
+		}
+	};
+
+	static auto ReimportAction = [](const FToolMenuContext& InMenuContext, const int32 SourceFileIndex, const bool bReimportAll, const bool bWithNewFile)
+	{
+		TSharedPtr<FSkeletalMeshEditor> SkeletalMeshEditor = GetSkeletalMeshEditor(InMenuContext);
+		if (SkeletalMeshEditor.IsValid())
+		{
+			if (USkeletalMesh* FoundSkeletalMesh = SkeletalMeshEditor->SkeletalMesh)
+			{
+				UFbxSkeletalMeshImportData* SkeletalMeshImportData = FoundSkeletalMesh ? Cast<UFbxSkeletalMeshImportData>(FoundSkeletalMesh->AssetImportData) : nullptr;
+				if (SkeletalMeshImportData)
+				{
+					SkeletalMeshImportData->ImportContentType = SourceFileIndex == 0 ? EFBXImportContentType::FBXICT_All : SourceFileIndex == 1 ? EFBXImportContentType::FBXICT_Geometry : EFBXImportContentType::FBXICT_SkinningWeights;
+					if (bReimportAll)
 					{
-						SkeletalMeshImportData->ImportContentType = SourceFileIndex == 0 ? EFBXImportContentType::FBXICT_All : SourceFileIndex == 1 ? EFBXImportContentType::FBXICT_Geometry : EFBXImportContentType::FBXICT_SkinningWeights;
-						HandleReimportMeshWithNewFile(SourceFileIndex);
+						if (bWithNewFile)
+						{
+							SkeletalMeshEditor->HandleReimportAllMeshWithNewFile(SourceFileIndex);
+						}
+						else
+						{
+							SkeletalMeshEditor->HandleReimportAllMesh(SourceFileIndex);
+						}
 					}
-				};
-
-				FText ReimportGeometryContentLabel = LOCTEXT("ReimportGeometryContentLabel", "Geometry");
-				FText ReimportGeometryContentLabelTooltip = LOCTEXT("ReimportGeometryContentLabelTooltipTooltip", "Reimport Geometry Only");
-				SubMenuBuilder.AddMenuEntry(
-					ReimportGeometryContentLabel,
-					ReimportGeometryContentLabelTooltip,
-					FSlateIcon(FEditorStyle::GetStyleSetName(), "ContentBrowser.AssetActions.ReimportAsset"),
-					FUIAction(
-						FExecuteAction::CreateLambda(UpdateImportDataContentType, 1),
-						FCanExecuteAction()
-					)
-				);
-				FText ReimportSkinningAndWeightsContentLabel = LOCTEXT("ReimportSkinningAndWeightsContentLabel", "Skinning And Weights");
-				FText ReimportSkinningAndWeightsContentLabelTooltip = LOCTEXT("ReimportSkinningAndWeightsContentLabelTooltipTooltip", "Reimport Skinning And Weights Only");
-				SubMenuBuilder.AddMenuEntry(
-					ReimportSkinningAndWeightsContentLabel,
-					ReimportSkinningAndWeightsContentLabelTooltip,
-					FSlateIcon(FEditorStyle::GetStyleSetName(), "ContentBrowser.AssetActions.ReimportAsset"),
-					FUIAction(
-						FExecuteAction::CreateLambda(UpdateImportDataContentType, 2),
-						FCanExecuteAction()
-					)
-				);
-			};
-
-			MenuBuilder.AddSubMenu(
-				ReimportMultiSources,
-				ReimportMultiSourcesTooltip,
-				FNewMenuDelegate::CreateLambda(CreateMultiContentSubMenu));
+					else
+					{
+						if (bWithNewFile)
+						{
+							SkeletalMeshEditor->HandleReimportMeshWithNewFile(SourceFileIndex);
+						}
+						else
+						{
+							SkeletalMeshEditor->HandleReimportMesh(SourceFileIndex);
+						}
+					}
+				}
+			}
 		}
-		else
+	};
+
+	static auto CreateReimportSubMenu = [](UToolMenu* InMenu, bool bReimportAll, bool bWithNewFile)
+	{
+		TSharedPtr<FSkeletalMeshEditor> SkeletalMeshEditor = GetSkeletalMeshEditor(InMenu->Context);
+		if (SkeletalMeshEditor.IsValid())
 		{
-			auto CreateSubMenu = [this](FMenuBuilder& SubMenuBuilder, bool bReimportAll, bool bWithNewFile)
+			USkeletalMesh* SkeletalMesh = SkeletalMeshEditor->SkeletalMesh;
+			if (SkeletalMesh && SkeletalMesh->AssetImportData)
 			{
-				//Get the data, we cannot use the closure since the lambda will be call when the function scope will be gone
+				//Get the data
 				TArray<FString> SourceFilePaths;
 				SkeletalMesh->AssetImportData->ExtractFilenames(SourceFilePaths);
 				TArray<FString> SourceFileLabels;
@@ -382,37 +400,7 @@ void FSkeletalMeshEditor::ExtendToolbar()
 
 				if (SourceFileLabels.Num() > 0 && SourceFileLabels.Num() == SourceFilePaths.Num())
 				{
-					auto UpdateImportDataContentType = [this](int32 SourceFileIndex, bool bReimportAll, bool bWithNewFile)
-					{
-						UFbxSkeletalMeshImportData* SkeletalMeshImportData = Cast<UFbxSkeletalMeshImportData>(SkeletalMesh->AssetImportData);
-						if (SkeletalMeshImportData)
-						{
-							SkeletalMeshImportData->ImportContentType = SourceFileIndex == 0 ? EFBXImportContentType::FBXICT_All : SourceFileIndex == 1 ? EFBXImportContentType::FBXICT_Geometry : EFBXImportContentType::FBXICT_SkinningWeights;
-							if (bReimportAll)
-							{
-								if (bWithNewFile)
-								{
-									HandleReimportAllMeshWithNewFile(SourceFileIndex);
-								}
-								else
-								{
-									HandleReimportAllMesh(SourceFileIndex);
-								}
-							}
-							else
-							{
-								if (bWithNewFile)
-								{
-									HandleReimportMeshWithNewFile(SourceFileIndex);
-								}
-								else
-								{
-									HandleReimportMesh(SourceFileIndex);
-								}
-							}
-						}
-					};
-
+					FToolMenuSection& Section = InMenu->AddSection("Reimport");
 					for (int32 SourceFileIndex = 0; SourceFileIndex < SourceFileLabels.Num(); ++SourceFileIndex)
 					{
 						FText ReimportLabel = FText::Format(LOCTEXT("ReimportNoLabel", "SourceFile {0}"), SourceFileIndex);
@@ -422,77 +410,112 @@ void FSkeletalMeshEditor::ExtendToolbar()
 							ReimportLabel = FText::Format(LOCTEXT("ReimportLabel", "{0}"), FText::FromString(SourceFileLabels[SourceFileIndex]));
 							ReimportLabelTooltip = FText::Format(LOCTEXT("ReimportLabelTooltip", "Reimport {0} File: {1}"), FText::FromString(SourceFileLabels[SourceFileIndex]), FText::FromString(SourceFilePaths[SourceFileIndex]));
 						}
-						SubMenuBuilder.AddMenuEntry(
+
+						Section.AddMenuEntry(
+							NAME_None,
 							ReimportLabel,
 							ReimportLabelTooltip,
 							FSlateIcon(FEditorStyle::GetStyleSetName(), "ContentBrowser.AssetActions.ReimportAsset"),
-							FUIAction(
-								FExecuteAction::CreateLambda(UpdateImportDataContentType, SourceFileIndex, bReimportAll, bWithNewFile),
-								FCanExecuteAction()
-							)
+							FToolMenuExecuteAction::CreateLambda(ReimportAction, SourceFileIndex, bReimportAll, bWithNewFile)
 						);
 					}
 				}
-			};
-
-			//Create 4 submenu: Reimport, ReimportWithNewFile, ReimportAll and ReimportAllWithNewFile
-			MenuBuilder.AddSubMenu(
-				FSkeletalMeshEditorCommands::Get().ReimportMesh->GetLabel(),
-				FSkeletalMeshEditorCommands::Get().ReimportMesh->GetDescription(),
-				FNewMenuDelegate::CreateLambda(CreateSubMenu, false, false));
-			
-			MenuBuilder.AddSubMenu(
-				FSkeletalMeshEditorCommands::Get().ReimportMeshWithNewFile->GetLabel(),
-				FSkeletalMeshEditorCommands::Get().ReimportMeshWithNewFile->GetDescription(),
-				FNewMenuDelegate::CreateLambda(CreateSubMenu, false, true));
-
-			MenuBuilder.AddSubMenu(
-				FSkeletalMeshEditorCommands::Get().ReimportAllMesh->GetLabel(),
-				FSkeletalMeshEditorCommands::Get().ReimportAllMesh->GetDescription(),
-				FNewMenuDelegate::CreateLambda(CreateSubMenu, true, false));
-
-			MenuBuilder.AddSubMenu(
-				FSkeletalMeshEditorCommands::Get().ReimportAllMeshWithNewFile->GetLabel(),
-				FSkeletalMeshEditorCommands::Get().ReimportAllMeshWithNewFile->GetDescription(),
-				FNewMenuDelegate::CreateLambda(CreateSubMenu, true, true));
+			}
 		}
-
-		return MenuBuilder.MakeWidget();
 	};
 
-	// extend extra menu/toolbars
-	auto FillToolbar = [this, ConstructReimportContextMenu](FToolBarBuilder& ToolbarBuilder)
+	if (!UToolMenus::Get()->IsMenuRegistered(UToolMenus::JoinMenuPaths(InBaseMenuName, "ReimportContextMenu")))
 	{
-		FPersonaModule& PersonaModule = FModuleManager::LoadModuleChecked<FPersonaModule>("Persona");
-		FPersonaModule::FCommonToolbarExtensionArgs Args;
-		Args.bPreviewMesh = false;
-		PersonaModule.AddCommonToolbarExtensions(ToolbarBuilder, PersonaToolkit.ToSharedRef(), Args);
-
-		ToolbarBuilder.BeginSection("Mesh");
+		UToolMenu* ToolMenu = UToolMenus::Get()->RegisterMenu(UToolMenus::JoinMenuPaths(InBaseMenuName, "ReimportContextMenu"));
+		ToolMenu->AddDynamicSection("Section", FNewToolMenuDelegate::CreateLambda([](UToolMenu* InMenu)
 		{
-			ToolbarBuilder.AddToolBarButton(FSkeletalMeshEditorCommands::Get().ReimportMesh);
-			ToolbarBuilder.AddComboButton(
-				FUIAction(),
-				FOnGetContent::CreateLambda(ConstructReimportContextMenu),
-				TAttribute<FText>(),
-				TAttribute<FText>()
-			);
-		}
-		ToolbarBuilder.EndSection();
+			TSharedPtr<FSkeletalMeshEditor> SkeletalMeshEditor = GetSkeletalMeshEditor(InMenu->Context);
+			if (SkeletalMeshEditor.IsValid())
+			{
+				USkeletalMesh* SkeletalMesh = SkeletalMeshEditor->SkeletalMesh;
+				bool bShowSubMenu = SkeletalMesh != nullptr && SkeletalMesh->AssetImportData != nullptr && SkeletalMesh->AssetImportData->GetSourceFileCount() > 1;
 
-		ToolbarBuilder.BeginSection("Selection");
-		{
-			ToolbarBuilder.AddToolBarButton(FSkeletalMeshEditorCommands::Get().MeshSectionSelection);
-		}
-		ToolbarBuilder.EndSection();
-	};
+				FToolMenuSection& Section = InMenu->AddSection("Section");
+				if (!bShowSubMenu)
+				{
+					//Reimport
+					Section.AddMenuEntry(
+						FSkeletalMeshEditorCommands::Get().ReimportMesh->GetCommandName(),
+						FSkeletalMeshEditorCommands::Get().ReimportMesh->GetLabel(),
+						FSkeletalMeshEditorCommands::Get().ReimportMesh->GetDescription(),
+						FSlateIcon(),
+						FUIAction(FExecuteAction::CreateSP(SkeletalMeshEditor.ToSharedRef(), &FSkeletalMeshEditor::HandleReimportMesh, 0)));
+
+					Section.AddMenuEntry(
+						FSkeletalMeshEditorCommands::Get().ReimportMeshWithNewFile->GetCommandName(),
+						FSkeletalMeshEditorCommands::Get().ReimportMeshWithNewFile->GetLabel(),
+						FSkeletalMeshEditorCommands::Get().ReimportMeshWithNewFile->GetDescription(),
+						FSlateIcon(),
+						FUIAction(FExecuteAction::CreateSP(SkeletalMeshEditor.ToSharedRef(), &FSkeletalMeshEditor::HandleReimportMeshWithNewFile, 0)));
+
+					//Reimport ALL
+					Section.AddMenuEntry(
+						FSkeletalMeshEditorCommands::Get().ReimportAllMesh->GetCommandName(),
+						FSkeletalMeshEditorCommands::Get().ReimportAllMesh->GetLabel(),
+						FSkeletalMeshEditorCommands::Get().ReimportAllMesh->GetDescription(),
+						FSlateIcon(),
+						FUIAction(FExecuteAction::CreateSP(SkeletalMeshEditor.ToSharedRef(), &FSkeletalMeshEditor::HandleReimportAllMesh, 0)));
+
+					Section.AddMenuEntry(
+						FSkeletalMeshEditorCommands::Get().ReimportAllMeshWithNewFile->GetCommandName(),
+						FSkeletalMeshEditorCommands::Get().ReimportAllMeshWithNewFile->GetLabel(),
+						FSkeletalMeshEditorCommands::Get().ReimportAllMeshWithNewFile->GetDescription(),
+						FSlateIcon(),
+						FUIAction(FExecuteAction::CreateSP(SkeletalMeshEditor.ToSharedRef(), &FSkeletalMeshEditor::HandleReimportAllMeshWithNewFile, 0)));
+
+					Section.AddSubMenu(
+						"ReimportMultiSources",
+						LOCTEXT("ReimportMultiSources", "Reimport Content"),
+						LOCTEXT("ReimportMultiSourcesTooltip", "Reimport Geometry or Skinning Weights content, this will create multi import source file."),
+						FNewToolMenuDelegate::CreateLambda(CreateMultiContentSubMenu));
+				}
+				else
+				{
+					//Create 4 submenu: Reimport, ReimportWithNewFile, ReimportAll and ReimportAllWithNewFile
+					Section.AddSubMenu(
+						FSkeletalMeshEditorCommands::Get().ReimportMesh->GetCommandName(),
+						FSkeletalMeshEditorCommands::Get().ReimportMesh->GetLabel(),
+						FSkeletalMeshEditorCommands::Get().ReimportMesh->GetDescription(),
+						FNewToolMenuDelegate::CreateLambda(CreateReimportSubMenu, false, false));
+			
+					Section.AddSubMenu(
+						FSkeletalMeshEditorCommands::Get().ReimportMeshWithNewFile->GetCommandName(),
+						FSkeletalMeshEditorCommands::Get().ReimportMeshWithNewFile->GetLabel(),
+						FSkeletalMeshEditorCommands::Get().ReimportMeshWithNewFile->GetDescription(),
+						FNewToolMenuDelegate::CreateLambda(CreateReimportSubMenu, false, true));
+
+					Section.AddSubMenu(
+						FSkeletalMeshEditorCommands::Get().ReimportAllMesh->GetCommandName(),
+						FSkeletalMeshEditorCommands::Get().ReimportAllMesh->GetLabel(),
+						FSkeletalMeshEditorCommands::Get().ReimportAllMesh->GetDescription(),
+						FNewToolMenuDelegate::CreateLambda(CreateReimportSubMenu, true, false));
+
+					Section.AddSubMenu(
+						FSkeletalMeshEditorCommands::Get().ReimportAllMeshWithNewFile->GetCommandName(),
+						FSkeletalMeshEditorCommands::Get().ReimportAllMeshWithNewFile->GetLabel(),
+						FSkeletalMeshEditorCommands::Get().ReimportAllMeshWithNewFile->GetDescription(),
+						FNewToolMenuDelegate::CreateLambda(CreateReimportSubMenu, true, true));
+				}
+			}
+		}));
+	}
+}
 
-	ToolbarExtender->AddToolBarExtension(
-		"Asset",
-		EExtensionHook::After,
-		GetToolkitCommands(),
-		FToolBarExtensionDelegate::CreateLambda(FillToolbar)
-		);
+void FSkeletalMeshEditor::ExtendToolbar()
+{
+	// If the ToolbarExtender is valid, remove it before rebuilding it
+	if (ToolbarExtender.IsValid())
+	{
+		RemoveToolbarExtender(ToolbarExtender);
+		ToolbarExtender.Reset();
+	}
+
+	ToolbarExtender = MakeShareable(new FExtender);
 
 	AddToolbarExtender(ToolbarExtender);
 
@@ -515,6 +538,7 @@ void FSkeletalMeshEditor::ExtendToolbar()
 		GetToolkitCommands(),
 		FToolBarExtensionDelegate::CreateLambda([this](FToolBarBuilder& ParentToolbarBuilder)
 	{
+		// Second toolbar on right side
 		FPersonaModule& PersonaModule = FModuleManager::LoadModuleChecked<FPersonaModule>("Persona");
 		TSharedRef<class IAssetFamily> AssetFamily = PersonaModule.CreatePersonaAssetFamily(SkeletalMesh);
 		AddToolbarWidget(PersonaModule.CreateAssetFamilyShortcutWidget(SharedThis(this), AssetFamily));
@@ -522,6 +546,53 @@ void FSkeletalMeshEditor::ExtendToolbar()
 	));
 }
 
+void FSkeletalMeshEditor::RegisterToolbar()
+{
+	static const FName MenuName = "AssetEditor.SkeletalMeshEditor.ToolBar";
+	RegisterReimportContextMenu(MenuName);
+
+	if (!UToolMenus::Get()->IsMenuRegistered(MenuName))
+	{
+		UToolMenu* ToolMenu = UToolMenus::Get()->RegisterMenu(MenuName, "AssetEditor.DefaultToolBar", EMultiBoxType::ToolBar);
+
+		const FToolMenuInsert SectionInsertLocation("Asset", EToolMenuInsertType::After);
+
+		{
+			ToolMenu->AddDynamicSection("Persona", FNewToolBarDelegateLegacy::CreateLambda([](FToolBarBuilder& ToolbarBuilder, UToolMenu* InMenu)
+			{
+				TSharedPtr<FSkeletalMeshEditor> SkeletalMeshEditor = GetSkeletalMeshEditor(InMenu->Context);
+				if (SkeletalMeshEditor.IsValid() && SkeletalMeshEditor->PersonaToolkit.IsValid())
+				{
+					FPersonaModule& PersonaModule = FModuleManager::LoadModuleChecked<FPersonaModule>("Persona");
+					FPersonaModule::FCommonToolbarExtensionArgs Args;
+					Args.bPreviewMesh = false;
+					PersonaModule.AddCommonToolbarExtensions(ToolbarBuilder, SkeletalMeshEditor->PersonaToolkit.ToSharedRef(), Args);
+				}
+			}), SectionInsertLocation);
+		}
+
+		{
+			FToolMenuSection& Section = ToolMenu->AddSection("Mesh", FText(), SectionInsertLocation);
+			Section.AddEntry(FToolMenuEntry::InitToolBarButton(FSkeletalMeshEditorCommands::Get().ReimportMesh));
+			Section.AddEntry(FToolMenuEntry::InitComboButton("ReimportContextMenu", FUIAction(), FNewToolMenuDelegate()));
+		}
+
+		{
+			FToolMenuSection& Section = ToolMenu->AddSection("SkeletalMesh", FText(), FToolMenuInsert("Mesh", EToolMenuInsertType::After));
+			Section.AddEntry(FToolMenuEntry::InitToolBarButton(FSkeletalMeshEditorCommands::Get().MeshSectionSelection));
+		}
+	}
+}
+
+void FSkeletalMeshEditor::InitToolMenuContext(FToolMenuContext& MenuContext)
+{
+	FAssetEditorToolkit::InitToolMenuContext(MenuContext);
+
+	USkeletalMeshToolMenuContext* Context = NewObject<USkeletalMeshToolMenuContext>();
+	Context->SkeletalMeshEditor = SharedThis(this);
+	MenuContext.AddObject(Context);
+}
+
 void FSkeletalMeshEditor::FillMeshClickMenu(FMenuBuilder& MenuBuilder, HActor* HitProxy, const FViewportClick& Click)
 {
 	UDebugSkelMeshComponent* MeshComp = GetPersonaToolkit()->GetPreviewMeshComponent();
diff --git a/Engine/Source/Editor/SkeletalMeshEditor/Private/SkeletalMeshEditor.h b/Engine/Source/Editor/SkeletalMeshEditor/Private/SkeletalMeshEditor.h
--- a/Engine/Source/Editor/SkeletalMeshEditor/Private/SkeletalMeshEditor.h
+++ b/Engine/Source/Editor/SkeletalMeshEditor/Private/SkeletalMeshEditor.h
@@ -21,6 +21,8 @@ class ISkeletonTreeItem;
 struct HActor;
 struct FViewportClick;
 struct FSkeletalMeshClothBuildParams;
+struct FToolMenuContext;
+class UToolMenu;
 
 namespace SkeletalMeshEditorModes
 {
@@ -61,6 +63,7 @@ public:
 	virtual FText GetBaseToolkitName() const override;
 	virtual FString GetWorldCentricTabPrefix() const override;
 	virtual FLinearColor GetWorldCentricTabColorScale() const override;
+	virtual void InitToolMenuContext(FToolMenuContext& MenuContext) override;
 
 	//~ Begin FAssetEditorToolkit Interface.
 	virtual bool OnRequestClose() override;
@@ -137,9 +140,14 @@ private:
 	void OnRemoveSectionFromLodAndBelowMenuItemClicked(int32 LodIndex, int32 SectionIndex);
 	//////////////////////////////////////////////////////////////////////////
 
+	void RegisterReimportContextMenu(const FName InBaseMenuName);
+
+	static TSharedPtr<FSkeletalMeshEditor> GetSkeletalMeshEditor(const FToolMenuContext& InMenuContext);
+
 private:
 	void ExtendMenu();
 
+	void RegisterToolbar();
 	void ExtendToolbar();
 
 	void BindCommands();
diff --git a/Engine/Source/Editor/SkeletalMeshEditor/Private/SkeletalMeshEditorMode.cpp b/Engine/Source/Editor/SkeletalMeshEditor/Private/SkeletalMeshEditorMode.cpp
--- a/Engine/Source/Editor/SkeletalMeshEditor/Private/SkeletalMeshEditorMode.cpp
+++ b/Engine/Source/Editor/SkeletalMeshEditor/Private/SkeletalMeshEditorMode.cpp
@@ -85,6 +85,8 @@ FSkeletalMeshEditorMode::FSkeletalMeshEditorMode(TSharedRef<FWorkflowCentricAppl
 	LayoutExtender = MakeShared<FLayoutExtender>();
 	PersonaModule.OnRegisterLayoutExtensions().Broadcast(*LayoutExtender.Get());
 	TabLayout->ProcessExtensions(*LayoutExtender.Get());
+
+	SkeletalMeshEditor->RegisterModeToolbarIfUnregistered(GetModeName());
 }
 
 void FSkeletalMeshEditorMode::RegisterTabFactories(TSharedPtr<FTabManager> InTabManager)
diff --git a/Engine/Source/Editor/SkeletalMeshEditor/SkeletalMeshEditor.Build.cs b/Engine/Source/Editor/SkeletalMeshEditor/SkeletalMeshEditor.Build.cs
--- a/Engine/Source/Editor/SkeletalMeshEditor/SkeletalMeshEditor.Build.cs
+++ b/Engine/Source/Editor/SkeletalMeshEditor/SkeletalMeshEditor.Build.cs
@@ -35,6 +35,7 @@ public class SkeletalMeshEditor : ModuleRules
                 "ClothingSystemEditorInterface",
 				"ClothingSystemRuntimeInterface",
 				"SkeletalMeshUtilitiesCommon",
+				"ToolMenus",
             }
 		);
 
diff --git a/Engine/Source/Editor/UnrealEd/Private/Toolkits/AssetEditorToolkit.cpp b/Engine/Source/Editor/UnrealEd/Private/Toolkits/AssetEditorToolkit.cpp
--- a/Engine/Source/Editor/UnrealEd/Private/Toolkits/AssetEditorToolkit.cpp
+++ b/Engine/Source/Editor/UnrealEd/Private/Toolkits/AssetEditorToolkit.cpp
@@ -986,7 +986,7 @@ FName FAssetEditorToolkit::GetToolMenuToolbarName(FName& OutParentName) const
 	return *(TEXT("AssetEditor.") + GetToolMenuAppName().ToString() + TEXT(".ToolBar"));
 }
 
-void FAssetEditorToolkit::RegisterMenus()
+void FAssetEditorToolkit::RegisterDefaultToolBar()
 {
 	static const FName DefaultToolBarName("AssetEditor.DefaultToolBar");
 	UToolMenus* ToolMenus = UToolMenus::Get();
@@ -1010,7 +1010,7 @@ void FAssetEditorToolkit::GenerateToolbar()
 {
 	TSharedPtr<FExtender> Extender = FExtender::Combine(ToolbarExtenders);
 
-	RegisterMenus();
+	RegisterDefaultToolBar();
 
 	FName ParentToolbarName;
 	const FName ToolBarName = GetToolMenuToolbarName(ParentToolbarName);
diff --git a/Engine/Source/Editor/UnrealEd/Public/SkeletalMeshToolMenuContext.h b/Engine/Source/Editor/UnrealEd/Public/SkeletalMeshToolMenuContext.h
new file mode 100644
--- /dev/null
+++ b/Engine/Source/Editor/UnrealEd/Public/SkeletalMeshToolMenuContext.h
@@ -0,0 +1,18 @@
+// Copyright 1998-2019 Epic Games, Inc. All Rights Reserved.
+
+#pragma once
+
+#include "CoreMinimal.h"
+
+#include "SkeletalMeshToolMenuContext.generated.h"
+
+class ISkeletalMeshEditor;
+
+UCLASS()
+class UNREALED_API USkeletalMeshToolMenuContext : public UObject
+{
+	GENERATED_BODY()
+public:
+
+	TWeakPtr<ISkeletalMeshEditor> SkeletalMeshEditor;
+};
diff --git a/Engine/Source/Editor/UnrealEd/Public/Toolkits/AssetEditorToolkit.h b/Engine/Source/Editor/UnrealEd/Public/Toolkits/AssetEditorToolkit.h
--- a/Engine/Source/Editor/UnrealEd/Public/Toolkits/AssetEditorToolkit.h
+++ b/Engine/Source/Editor/UnrealEd/Public/Toolkits/AssetEditorToolkit.h
@@ -179,7 +179,7 @@ public:
 	}
 
 	/** Registers default tool bar */
-	static void RegisterMenus();
+	static void RegisterDefaultToolBar();
 	
 	/** Makes a default asset editing toolbar */
 	void GenerateToolbar();
