Fix mac compile warning
#rnx

diff --git a/Engine/Source/Editor/SkeletalMeshEditor/Private/SkeletalMeshEditor.cpp b/Engine/Source/Editor/SkeletalMeshEditor/Private/SkeletalMeshEditor.cpp
--- a/Engine/Source/Editor/SkeletalMeshEditor/Private/SkeletalMeshEditor.cpp
+++ b/Engine/Source/Editor/SkeletalMeshEditor/Private/SkeletalMeshEditor.cpp
@@ -389,14 +389,14 @@ void FSkeletalMeshEditor::RegisterReimportContextMenu(const FName InBaseMenuName
 		TSharedPtr<FSkeletalMeshEditor> SkeletalMeshEditor = GetSkeletalMeshEditor(InMenu->Context);
 		if (SkeletalMeshEditor.IsValid())
 		{
-			USkeletalMesh* SkeletalMesh = SkeletalMeshEditor->SkeletalMesh;
-			if (SkeletalMesh && SkeletalMesh->AssetImportData)
+			USkeletalMesh* InSkeletalMesh = SkeletalMeshEditor->SkeletalMesh;
+			if (InSkeletalMesh && InSkeletalMesh->AssetImportData)
 			{
 				//Get the data
 				TArray<FString> SourceFilePaths;
-				SkeletalMesh->AssetImportData->ExtractFilenames(SourceFilePaths);
+				InSkeletalMesh->AssetImportData->ExtractFilenames(SourceFilePaths);
 				TArray<FString> SourceFileLabels;
-				SkeletalMesh->AssetImportData->ExtractDisplayLabels(SourceFileLabels);
+				InSkeletalMesh->AssetImportData->ExtractDisplayLabels(SourceFileLabels);
 
 				if (SourceFileLabels.Num() > 0 && SourceFileLabels.Num() == SourceFilePaths.Num())
 				{
@@ -432,8 +432,8 @@ void FSkeletalMeshEditor::RegisterReimportContextMenu(const FName InBaseMenuName
 			TSharedPtr<FSkeletalMeshEditor> SkeletalMeshEditor = GetSkeletalMeshEditor(InMenu->Context);
 			if (SkeletalMeshEditor.IsValid())
 			{
-				USkeletalMesh* SkeletalMesh = SkeletalMeshEditor->SkeletalMesh;
-				bool bShowSubMenu = SkeletalMesh != nullptr && SkeletalMesh->AssetImportData != nullptr && SkeletalMesh->AssetImportData->GetSourceFileCount() > 1;
+				USkeletalMesh* InSkeletalMesh = SkeletalMeshEditor->SkeletalMesh;
+				bool bShowSubMenu = InSkeletalMesh != nullptr && InSkeletalMesh->AssetImportData != nullptr && InSkeletalMesh->AssetImportData->GetSourceFileCount() > 1;
 
 				FToolMenuSection& Section = InMenu->AddSection("Section");
 				if (!bShowSubMenu)
