Chaos Cloth - Fix LOD reset after undo. [ Merging CL 14421305]

diff --git a/Engine/Source/Editor/SkeletalMeshEditor/Private/SkeletalMeshEditor.cpp b/Engine/Source/Editor/SkeletalMeshEditor/Private/SkeletalMeshEditor.cpp
--- a/Engine/Source/Editor/SkeletalMeshEditor/Private/SkeletalMeshEditor.cpp
+++ b/Engine/Source/Editor/SkeletalMeshEditor/Private/SkeletalMeshEditor.cpp
@@ -1042,6 +1042,7 @@ void FSkeletalMeshEditor::RemoveClothing(int32 InLodIndex, int32 InSectionIndex)
 		// Look for a currently bound asset an unbind it if necessary first
 		if (UClothingAssetBase* CurrentAsset = Mesh->GetSectionClothingAsset(InLodIndex, InSectionIndex))
 		{
+			CurrentAsset->Modify();
 			CurrentAsset->UnbindFromSkeletalMesh(Mesh, InLodIndex);
 			ClearOriginalSectionUserData();
 		}
