Add editor mode toolkit support to the skeletal mesh editor.

diff --git a/Engine/Source/Editor/SkeletalMeshEditor/Private/SSkeletalMeshEditorToolbox.cpp b/Engine/Source/Editor/SkeletalMeshEditor/Private/SSkeletalMeshEditorToolbox.cpp
new file mode 100644
--- /dev/null
+++ b/Engine/Source/Editor/SkeletalMeshEditor/Private/SSkeletalMeshEditorToolbox.cpp
@@ -0,0 +1,237 @@
+// Copyright Epic Games, Inc. All Rights Reserved.
+
+#include "SSkeletalMeshEditorToolbox.h"
+
+#include "ISkeletalMeshEditor.h"
+
+#include "EdMode.h"
+#include "Framework/MultiBox/MultiBoxBuilder.h"
+#include "Tools/UEdMode.h"
+#include "Widgets/Input/SCheckBox.h"
+#include "Widgets/Layout/SWidgetSwitcher.h"
+#include "Widgets/Layout/SUniformWrapPanel.h"
+
+
+SSkeletalMeshEditorToolbox::~SSkeletalMeshEditorToolbox()
+{
+}
+
+
+void SSkeletalMeshEditorToolbox::Construct(
+	const FArguments& InArgs, 
+	const TSharedRef<ISkeletalMeshEditor>& InOwningEditor
+	)
+{
+	SkeletalMeshEditor = InOwningEditor;
+
+	ChildSlot
+	[
+		SNew( SBorder )
+		.BorderImage( FEditorStyle::GetBrush( "ToolPanel.GroupBorder" ) )
+		.Padding(0.0f)
+		[
+			SNew( SVerticalBox )
+			+ SVerticalBox::Slot()
+			.AutoHeight()
+			.HAlign( HAlign_Left )
+			[
+				SAssignNew( ModeToolBarContainer, SBorder )
+				.BorderImage( FEditorStyle::GetBrush( "ToolPanel.GroupBorder" ) )
+				.Padding( FMargin(4, 0, 0, 0) )
+			]
+
+			+ SVerticalBox::Slot()
+			.FillHeight( 1.0f )
+			[
+				SNew( SVerticalBox )
+
+				+SVerticalBox::Slot()
+				.Padding(0.0, 8.0, 0.0, 0.0)
+				.AutoHeight()
+				[
+					SAssignNew(ModeToolHeader, SBorder)
+					.BorderImage( FEditorStyle::GetBrush( "ToolPanel.GroupBorder" ) )
+				]
+
+				+ SVerticalBox::Slot()
+				.FillHeight(1)
+				[
+					SAssignNew(InlineContentHolder, SBorder)
+					.BorderImage( FEditorStyle::GetBrush( "ToolPanel.GroupBorder" ) )
+					.Visibility( this, &SSkeletalMeshEditorToolbox::GetInlineContentHolderVisibility )
+				]
+			]
+		]
+	];
+}
+
+
+void SSkeletalMeshEditorToolbox::AttachToolkit(const TSharedRef<IToolkit>& InToolkit)
+{
+	UpdateInlineContent(InToolkit, InToolkit->GetInlineContent());
+}
+
+
+void SSkeletalMeshEditorToolbox::DetachToolkit(const TSharedRef<IToolkit>& InToolkit)
+{
+	UpdateInlineContent(nullptr, SNullWidget::NullWidget);
+}
+
+
+void SSkeletalMeshEditorToolbox::SetOwningTab(TSharedRef<SDockTab>& InOwningTab)
+{
+	OwningTab = InOwningTab;
+}
+
+
+void SSkeletalMeshEditorToolbox::UpdateInlineContent(const TSharedPtr<IToolkit>& Toolkit, TSharedPtr<SWidget> InlineContent)
+{
+	static const FName SkeletalMeshEditorStatusBarName = "SkeletalMeshEditor.StatusBar";
+
+	// The display name that the owning tab should have as its label
+	FText TabName;
+
+	// The icon that should be displayed in the parent tab
+	const FSlateBrush* TabIcon = nullptr;
+
+	if (StatusBarMessageHandle.IsValid())
+	{
+		GEditor->GetEditorSubsystem<UStatusBarSubsystem>()->PopStatusBarMessage(SkeletalMeshEditorStatusBarName, StatusBarMessageHandle);
+		StatusBarMessageHandle.Reset();
+	}
+
+	if (Toolkit.IsValid())
+	{
+		TabName = Toolkit->GetEditorModeDisplayName();
+		TabIcon = Toolkit->GetEditorModeIcon().GetSmallIcon();
+
+		TWeakPtr<FModeToolkit> ModeToolkit = StaticCastSharedPtr<FModeToolkit>(Toolkit);
+
+		if (ModeToolkit.IsValid())
+		{
+			TSharedRef<FModeToolkit> ModeToolkitPinned = ModeToolkit.Pin().ToSharedRef();
+
+			UpdatePalette(ModeToolkitPinned);
+
+			// Show the name of the active tool in the statusbar.
+			// FIXME: We should also be showing Ctrl/Shift/Alt LMB/RMB shortcuts.
+			StatusBarMessageHandle = GEditor->GetEditorSubsystem<UStatusBarSubsystem>()->PushStatusBarMessage(
+					SkeletalMeshEditorStatusBarName,
+					TAttribute<FText>::Create(TAttribute<FText>::FGetter::CreateSP(ModeToolkitPinned, &FModeToolkit::GetActiveToolDisplayName)));
+		}
+	}
+	else
+	{
+		TabName = NSLOCTEXT("SkeletalMeshEditor", "ToolboxTab", "Toolbox");
+		TabIcon = FEditorStyle::Get().GetBrush("LevelEditor.Tabs.Modes");
+	}
+
+	if (InlineContent.IsValid() && InlineContentHolder.IsValid())
+	{
+		InlineContentHolder->SetContent(InlineContent.ToSharedRef());
+	}
+
+	TSharedPtr<SDockTab> OwningTabPinned = OwningTab.Pin();
+	if (OwningTabPinned.IsValid())
+	{
+		OwningTabPinned->SetLabel(TabName);
+		OwningTabPinned->SetTabIcon(TabIcon);
+	}
+}
+
+
+void SSkeletalMeshEditorToolbox::UpdatePalette(const TSharedRef<FModeToolkit>& InModeToolkit)
+{
+	TSharedRef<SUniformWrapPanel> PaletteTabBox = SNew(SUniformWrapPanel)
+	                                                  .SlotPadding(FMargin(1.f, 2.f))
+	                                                  .HAlign(HAlign_Center);
+
+	// Only show if there's more than one entry.
+	PaletteTabBox->SetVisibility(TAttribute<EVisibility>::Create(
+		TAttribute<EVisibility>::FGetter::CreateLambda([PaletteTabBox]() -> EVisibility { 
+			return PaletteTabBox->GetChildren()->Num() > 1 ? EVisibility::Visible : EVisibility::Collapsed; 
+		})));
+
+	// Also build the toolkit here
+	TArray<FName> PaletteNames;
+	InModeToolkit->GetToolPaletteNames(PaletteNames);
+
+	TSharedPtr<FUICommandList> CommandList;
+	CommandList = InModeToolkit->GetToolkitCommands();
+
+	TSharedRef< SWidgetSwitcher > PaletteSwitcher = SNew(SWidgetSwitcher)
+	.WidgetIndex_Lambda( [PaletteNames, InModeToolkit] () -> int32 { 
+		int32 FoundIndex;
+		if (PaletteNames.Find(InModeToolkit->GetCurrentPalette(), FoundIndex))
+		{
+			return FoundIndex;	
+		}
+		return 0;
+	} );
+			
+	for(auto Palette : PaletteNames)
+	{
+		FName ToolbarCustomizationName = InModeToolkit->GetEditorMode() ? InModeToolkit->GetEditorMode()->GetModeInfo().ToolbarCustomizationName : InModeToolkit->GetScriptableEditorMode()->GetModeInfo().ToolbarCustomizationName;
+		FUniformToolBarBuilder ModeToolbarBuilder(CommandList, FMultiBoxCustomization(ToolbarCustomizationName));
+		ModeToolbarBuilder.SetStyle(&FEditorStyle::Get(), "PaletteToolBar");
+
+		InModeToolkit->BuildToolPalette(Palette, ModeToolbarBuilder);
+
+		TSharedRef<SWidget> PaletteWidget = ModeToolbarBuilder.MakeWidget();
+
+		PaletteTabBox->AddSlot()
+		[
+			SNew(SCheckBox)
+			.Padding(FMargin(8.f, 4.f, 8.f, 5.f))
+			.Style( FEditorStyle::Get(),  "PaletteToolBar.Tab" )
+			.OnCheckStateChanged_Lambda([/*PaletteSwitcher, PaletteWidget, */InModeToolkit, Palette](const ECheckBoxState) {
+					InModeToolkit->SetCurrentPalette(Palette);
+				} 
+			)
+			// .IsChecked_Lambda( [PaletteSwitcher, PaletteWidget] () -> ECheckBoxState { return PaletteSwitcher->GetActiveWidget() == PaletteWidget ? ECheckBoxState::Checked : ECheckBoxState::Unchecked; })
+			.IsChecked_Lambda([InModeToolkit, Palette]() -> ECheckBoxState {
+				if (InModeToolkit->GetCurrentPalette() == Palette)
+				{
+					return ECheckBoxState::Checked;
+				}
+				return ECheckBoxState::Unchecked;
+			})
+			[
+				SNew(STextBlock)
+				.TextStyle(FAppStyle::Get(), "NormalText")
+				.Text(InModeToolkit->GetToolPaletteDisplayName(Palette))
+				.Justification(ETextJustify::Center)
+			]
+		];
+
+
+		PaletteSwitcher->AddSlot()
+		[
+			PaletteWidget
+		]; 
+	}
+
+
+	ModeToolHeader->SetContent(
+		SNew(SVerticalBox)
+
+		+SVerticalBox::Slot()
+		.Padding(8.f, 0.f, 0.f, 8.f)
+		.AutoHeight()
+		[
+			PaletteTabBox
+		]
+
+		+SVerticalBox::Slot()
+		.AutoHeight()
+		[
+			PaletteSwitcher
+		]
+	);
+}
+
+
+EVisibility SSkeletalMeshEditorToolbox::GetInlineContentHolderVisibility() const
+{
+	return InlineContentHolder->GetContent() == SNullWidget::NullWidget ? EVisibility::Collapsed : EVisibility::Visible;
+}
diff --git a/Engine/Source/Editor/SkeletalMeshEditor/Private/SSkeletalMeshEditorToolbox.h b/Engine/Source/Editor/SkeletalMeshEditor/Private/SSkeletalMeshEditorToolbox.h
new file mode 100644
--- /dev/null
+++ b/Engine/Source/Editor/SkeletalMeshEditor/Private/SSkeletalMeshEditorToolbox.h
@@ -0,0 +1,65 @@
+// Copyright Epic Games, Inc. All Rights Reserved.
+
+#pragma once
+
+#include "ISkeletalMeshEditor.h"
+
+#include "CoreMinimal.h"
+#include "StatusBarSubsystem.h"
+#include "Widgets/DeclarativeSyntaxSupport.h"
+#include "Widgets/Docking/SDockTab.h"
+#include "Widgets/SCompoundWidget.h"
+
+
+class FModeToolkit;
+class ISkeletalMeshEditor;
+class IToolkit;
+class SBorder;
+
+
+class SSkeletalMeshEditorToolbox :
+	public SCompoundWidget
+{
+public:
+	SLATE_BEGIN_ARGS( SSkeletalMeshEditorToolbox ) {}
+	SLATE_END_ARGS()
+
+	~SSkeletalMeshEditorToolbox();
+
+	void Construct(const FArguments& InArgs, const TSharedRef<ISkeletalMeshEditor>& InOwningEditor);
+
+	// Called to attach a toolkit to this toolbox widget (e.g. when the mode changes).
+	void AttachToolkit(const TSharedRef<IToolkit>& InToolkit);
+
+	// Called to detach a toolkit from this toolbox widget (e.g. when the mode changes).
+	void DetachToolkit(const TSharedRef<IToolkit>& InToolkit);
+
+	// Sets the owning tab so that we can update the tab icon based on what the toolkit provides.
+	void SetOwningTab(TSharedRef<SDockTab>& InOwningTab);
+
+private:	
+	void UpdateInlineContent(const TSharedPtr<IToolkit>& Toolkit, TSharedPtr<SWidget> InlineContent);
+
+	void UpdatePalette(const TSharedRef<FModeToolkit>& InModeToolkit);
+
+	// Whether to show the inline content or not.
+	EVisibility GetInlineContentHolderVisibility() const;
+
+	// The owning tab where the toolkit's interface is hosted. We need this to update the icons
+	// and title as the toolkit changes.
+	TWeakPtr<SDockTab> OwningTab;
+
+	TWeakPtr<ISkeletalMeshEditor> SkeletalMeshEditor;
+
+	// Inline content area for editor modes
+	TSharedPtr<SBorder> InlineContentHolder;
+
+	// The container holding the mode toolbar
+	TSharedPtr<SBorder> ModeToolBarContainer;
+
+	// The active tool header area
+	TSharedPtr<SBorder> ModeToolHeader;
+
+	// The global status bar to show the name of the active tool.
+	FStatusBarMessageHandle StatusBarMessageHandle;
+};
\ No newline at end of file
diff --git a/Engine/Source/Editor/SkeletalMeshEditor/Private/SkeletalMeshEditor.cpp b/Engine/Source/Editor/SkeletalMeshEditor/Private/SkeletalMeshEditor.cpp
--- a/Engine/Source/Editor/SkeletalMeshEditor/Private/SkeletalMeshEditor.cpp
+++ b/Engine/Source/Editor/SkeletalMeshEditor/Private/SkeletalMeshEditor.cpp
@@ -1,52 +1,54 @@
 // Copyright Epic Games, Inc. All Rights Reserved.
 
 #include "SkeletalMeshEditor.h"
-#include "Modules/ModuleManager.h"
-#include "Framework/MultiBox/MultiBoxBuilder.h"
-#include "EditorReimportHandler.h"
+
+#include "SkeletalMeshEditorCommands.h"
+#include "SkeletalMeshEditorMode.h"
+#include "SSkeletalMeshEditorToolbox.h"
+
+
+#include "Algo/Transform.h"
+#include "Animation/DebugSkelMeshComponent.h"
 #include "AssetData.h"
+#include "ClothingAsset.h"
+#include "ClothingSystemEditorInterfaceModule.h"
+#include "ComponentReregisterContext.h"
 #include "EdGraph/EdGraphSchema.h"
 #include "Editor/EditorEngine.h"
+#include "EditorFramework/AssetImportData.h"
+#include "EditorModeManager.h"
+#include "EditorReimportHandler.h"
+#include "EditorViewportClient.h"
 #include "EngineGlobals.h"
-#include "ISkeletalMeshEditorModule.h"
-#include "IPersonaToolkit.h"
-#include "PersonaModule.h"
-#include "SkeletalMeshEditorMode.h"
-#include "IPersonaPreviewScene.h"
-#include "SkeletalMeshEditorCommands.h"
-#include "IDetailsView.h"
-#include "ISkeletonTree.h"
-#include "ISkeletonEditorModule.h"
-#include "IAssetFamily.h"
-#include "PersonaCommonCommands.h"
 #include "EngineUtils.h"
-#include "Rendering/SkeletalMeshModel.h"
-
-#include "Animation/DebugSkelMeshComponent.h"
-#include "ClothingAsset.h"
-#include "SCreateClothingSettingsPanel.h"
-#include "ClothingSystemEditorInterfaceModule.h"
-#include "Preferences/PersonaOptions.h"
-
-#include "Widgets/Layout/SBox.h"
-#include "Widgets/Layout/SBorder.h"
+#include "Factories/FbxSkeletalMeshImportData.h"
+#include "FbxMeshUtils.h"
 #include "Framework/Application/SlateApplication.h"
-#include "ToolMenus.h"
-#include "SkeletalMeshToolMenuContext.h"
-#include "EditorViewportClient.h"
-#include "Settings/EditorExperimentalSettings.h"
-#include "Algo/Transform.h"
+#include "Framework/MultiBox/MultiBoxBuilder.h"
+#include "IAssetFamily.h"
+#include "IDetailsView.h"
+#include "IPersonaPreviewScene.h"
+#include "IPersonaToolkit.h"
+#include "ISkeletalMeshEditorModule.h"
+#include "ISkeletonEditorModule.h"
+#include "ISkeletonTree.h"
 #include "ISkeletonTreeItem.h"
-#include "FbxMeshUtils.h"
 #include "LODUtilities.h"
-
-#include "ScopedTransaction.h"
-#include "ComponentReregisterContext.h"
-#include "EditorFramework/AssetImportData.h"
-#include "Factories/FbxSkeletalMeshImportData.h"
-
 #include "Misc/MessageDialog.h"
+#include "Modules/ModuleManager.h"
+#include "PersonaCommonCommands.h"
+#include "PersonaModule.h"
 #include "PersonaToolMenuContext.h"
+#include "Preferences/PersonaOptions.h"
+#include "Rendering/SkeletalMeshModel.h"
+#include "ScopedTransaction.h"
+#include "SCreateClothingSettingsPanel.h"
+#include "Settings/EditorExperimentalSettings.h"
+#include "SkeletalMeshToolMenuContext.h"
+#include "ToolMenus.h"
+#include "Widgets/Docking/SDockTab.h"
+#include "Widgets/Layout/SBorder.h"
+#include "Widgets/Layout/SBox.h"
 
 const FName SkeletalMeshEditorAppIdentifier = FName(TEXT("SkeletalMeshEditorApp"));
 
@@ -59,6 +61,8 @@
 const FName SkeletalMeshEditorTabs::AdvancedPreviewTab(TEXT("AdvancedPreviewTab"));
 const FName SkeletalMeshEditorTabs::MorphTargetsTab("MorphTargetsTab");
 const FName SkeletalMeshEditorTabs::AnimationMappingTab("AnimationMappingWindow");
+const FName SkeletalMeshEditorTabs::ToolboxDetailsTab("ToolBoxDetailsTab");
 
+
 DEFINE_LOG_CATEGORY(LogSkeletalMeshEditor);
 
@@ -190,6 +194,15 @@
 void FSkeletalMeshEditor::RegisterTabSpawners(const TSharedRef<class FTabManager>& InTabManager)
 {
 	WorkspaceMenuCategory = InTabManager->AddLocalWorkspaceMenuCategory(LOCTEXT("WorkspaceMenu_SkeletalMeshEditor", "Skeletal Mesh Editor"));
+	
+	InTabManager->RegisterTabSpawner(
+		SkeletalMeshEditorTabs::ToolboxDetailsTab, 
+		FOnSpawnTab::CreateSP(this, &FSkeletalMeshEditor::SpawnToolboxTab),
+		FCanSpawnTab::CreateSP(this, &FSkeletalMeshEditor::CanSpawnToolboxTab)
+		)
+		.SetDisplayName(LOCTEXT("ToolboxTab", "Toolbox"))
+		.SetGroup(WorkspaceMenuCategory.ToSharedRef())
+		.SetIcon(FSlateIcon(FEditorStyle::GetStyleSetName(), "LevelEditor.Tabs.Modes" ));
 
 	FAssetEditorToolkit::RegisterTabSpawners(InTabManager);
 }
@@ -240,6 +253,10 @@
 	// Set up mesh click selection
 	PreviewScene->RegisterOnMeshClick(FOnMeshClick::CreateSP(this, &FSkeletalMeshEditor::HandleMeshClick));
 	PreviewScene->SetAllowMeshHitProxies(GetDefault<UPersonaOptions>()->bAllowMeshSectionSelection);
+
+	// Make sure we get told when the editor mode changes so we can switch to the appropriate tab
+	// if there's a toolbox available.
+	GetEditorModeManager().OnEditorModeIDChanged().AddSP(this, &FSkeletalMeshEditor::OnEditorModeIdChanged);
 }
 
 FName FSkeletalMeshEditor::GetToolkitFName() const
@@ -299,6 +316,68 @@
 	return TSharedPtr<FSkeletalMeshEditor>();
 }
 
+
+void FSkeletalMeshEditor::OnEditorModeIdChanged(const FEditorModeID& ModeChangedID, bool bIsEnteringMode)
+{
+	if (GetEditorModeManager().IsDefaultMode(ModeChangedID))
+	{
+		return;
+	}
+
+	if (bIsEnteringMode)
+	{
+		// FIXME: We should get the hosted toolkit from here.
+		if (GetEditorModeManager().GetActiveScriptableMode(ModeChangedID)->UsesToolkits())
+		{
+			TabManager->TryInvokeTab(SkeletalMeshEditorTabs::ToolboxDetailsTab);
+		}
+	}
+	else
+	{
+		TSharedPtr<SDockTab> ToolboxTab = TabManager->FindExistingLiveTab(SkeletalMeshEditorTabs::ToolboxDetailsTab);
+		if (ToolboxTab.IsValid())
+		{
+			ToolboxTab->RequestCloseTab();
+		}
+	}
+}
+
+
+bool FSkeletalMeshEditor::CanSpawnToolboxTab(const FSpawnTabArgs& InArgs) const
+{
+	return HostedToolkit.IsValid();
+}
+
+
+TSharedRef<SDockTab> FSkeletalMeshEditor::SpawnToolboxTab(const FSpawnTabArgs& Args)
+{
+	ToolboxWidget = SNew(SSkeletalMeshEditorToolbox, SharedThis<ISkeletalMeshEditor>(this));
+
+	TSharedRef<SDockTab> DockTab = SNew(SDockTab)
+	    .Icon(FEditorStyle::GetBrush("LevelEditor.Tabs.Modes"))
+	    .Label(LOCTEXT("ToolboxTab", "Toolbox"))
+		[
+			SNew(SBox)
+	        .AddMetaData<FTagMetaData>(FTagMetaData(TEXT("ToolboxTab")))
+	        [
+				ToolboxWidget.ToSharedRef()
+			]
+		];
+
+	ToolboxWidget->SetOwningTab(DockTab);
+	ToolboxWidget->AttachToolkit(HostedToolkit.ToSharedRef());
+	
+	return DockTab;
+}
+
+
+void FSkeletalMeshEditor::OnToolboxTabClosed(TSharedRef<SDockTab> InClosedTab)
+{
+	// If the user closed the tab, then we want to go back to the base skel mesh editor mode.
+	GetEditorModeManager().ActivateDefaultMode();
+}
+
+
 void FSkeletalMeshEditor::RegisterReimportContextMenu(const FName InBaseMenuName)
 {
 	static auto ReimportMeshWithNewFileAction = [](const FToolMenuContext& InMenuContext, int32 SourceFileIndex)
@@ -594,6 +673,45 @@
 	MenuContext.AddObject(PersonaContext);
 }
 
+
+void FSkeletalMeshEditor::OnToolkitHostingStarted(const TSharedRef<IToolkit>& Toolkit)
+{
+	if (ensure(!HostedToolkit.IsValid()))
+	{
+		HostedToolkit = Toolkit;
+
+		if (ToolboxWidget.IsValid())
+		{
+			ToolboxWidget->AttachToolkit(Toolkit);
+		}
+	}
+}
+
+
+void FSkeletalMeshEditor::OnToolkitHostingFinished(const TSharedRef<IToolkit>& Toolkit)
+{
+	if (ensure(Toolkit == HostedToolkit))
+	{
+		if (ToolboxWidget.IsValid())
+		{
+			ToolboxWidget->DetachToolkit(Toolkit);
+		}
+		
+		HostedToolkit.Reset();
+	}
+}
+
+bool FSkeletalMeshEditor::ProcessCommandBindings(const FKeyEvent& InKeyEvent) const
+{
+	if (HostedToolkit.IsValid() && HostedToolkit->ProcessCommandBindings(InKeyEvent))
+	{
+		return true;
+	}
+
+	return ISkeletalMeshEditor::ProcessCommandBindings(InKeyEvent);
+}
+
+
 void FSkeletalMeshEditor::FillMeshClickMenu(FMenuBuilder& MenuBuilder, HActor* HitProxy, const FViewportClick& Click)
 {
 	UDebugSkelMeshComponent* MeshComp = GetPersonaToolkit()->GetPreviewMeshComponent();
diff --git a/Engine/Source/Editor/SkeletalMeshEditor/Private/SkeletalMeshEditor.h b/Engine/Source/Editor/SkeletalMeshEditor/Private/SkeletalMeshEditor.h
--- a/Engine/Source/Editor/SkeletalMeshEditor/Private/SkeletalMeshEditor.h
+++ b/Engine/Source/Editor/SkeletalMeshEditor/Private/SkeletalMeshEditor.h
@@ -23,6 +23,8 @@ struct FViewportClick;
 struct FSkeletalMeshClothBuildParams;
 struct FToolMenuContext;
 class UToolMenu;
+class SSkeletalMeshEditorToolbox;
+
 
 namespace SkeletalMeshEditorModes
 {
@@ -41,6 +43,7 @@ namespace SkeletalMeshEditorTabs
 	extern const FName MorphTargetsTab;
 	extern const FName MeshDetailsTab;
 	extern const FName AnimationMappingTab;
+	extern const FName ToolboxDetailsTab;
 }
 
 class FSkeletalMeshEditor : public ISkeletalMeshEditor, public FGCObject, public FEditorUndoClient, public FTickableEditorObject
@@ -64,7 +67,13 @@ public:
 	virtual FString GetWorldCentricTabPrefix() const override;
 	virtual FLinearColor GetWorldCentricTabColorScale() const override;
 	virtual void InitToolMenuContext(FToolMenuContext& MenuContext) override;
+	// IToolkitHost Interface
+	virtual void OnToolkitHostingStarted(const TSharedRef<class IToolkit>& Toolkit) override;
+	virtual void OnToolkitHostingFinished(const TSharedRef<class IToolkit>& Toolkit) override;
 
+	/** FBaseToolkit overrides */
+	virtual bool ProcessCommandBindings(const FKeyEvent& InKeyEvent) const override;
+	
 	//~ Begin FAssetEditorToolkit Interface.
 	virtual bool OnRequestClose() override;
 	//~ End FAssetEditorToolkit Interface.
@@ -96,6 +105,9 @@ public:
 
 	UObject* HandleGetAsset();
 
+	// Returns the currently hosted toolkit. Can be invalid if no toolkit is being hosted.
+	TSharedPtr<IToolkit> GetHostedToolkit() const { return HostedToolkit; }
+
 private:
 	void HandleObjectsSelected(const TArray<UObject*>& InObjects);
 
@@ -144,6 +156,12 @@ private:
 
 	static TSharedPtr<FSkeletalMeshEditor> GetSkeletalMeshEditor(const FToolMenuContext& InMenuContext);
 
+	void OnEditorModeIdChanged(const FEditorModeID& ModeChangedID, bool bIsEnteringMode);
+
+	bool CanSpawnToolboxTab(const FSpawnTabArgs& InArgs) const;
+	TSharedRef<SDockTab> SpawnToolboxTab(const FSpawnTabArgs& InArgs);
+	void OnToolboxTabClosed(TSharedRef<SDockTab> InClosedTab);
+
 private:
 	void ExtendMenu();
 
@@ -179,4 +197,10 @@ private:
 
 	/** Details panel */
 	TSharedPtr<class IDetailsView> DetailsView;
+
+	// The toolkit we're currently hosting.
+	TSharedPtr<IToolkit> HostedToolkit;
+
+	// The toolbox widget
+	TSharedPtr<SSkeletalMeshEditorToolbox> ToolboxWidget;
 };
diff --git a/Engine/Source/Editor/SkeletalMeshEditor/Private/SkeletalMeshEditorMode.cpp b/Engine/Source/Editor/SkeletalMeshEditor/Private/SkeletalMeshEditorMode.cpp
--- a/Engine/Source/Editor/SkeletalMeshEditor/Private/SkeletalMeshEditorMode.cpp
+++ b/Engine/Source/Editor/SkeletalMeshEditor/Private/SkeletalMeshEditorMode.cpp
@@ -35,7 +35,7 @@ FSkeletalMeshEditorMode::FSkeletalMeshEditorMode(TSharedRef<FWorkflowCentricAppl
 
 	TabFactories.RegisterFactory(CreateMeshControllerMappingTabFactory(InHostingApp, Cast<USkeletalMesh> (SkeletalMeshEditor->HandleGetAsset()), SkeletalMeshEditor->OnPostUndo));
 
-	TabLayout = FTabManager::NewLayout("Standalone_SkeletalMeshEditor_Layout_v3.3")
+	TabLayout = FTabManager::NewLayout("Standalone_SkeletalMeshEditor_Layout_v3.4")
 		->AddArea
 		(
 			FTabManager::NewPrimaryArea()
@@ -68,6 +68,7 @@ FSkeletalMeshEditorMode::FSkeletalMeshEditorMode(TSharedRef<FWorkflowCentricAppl
 					->SetHideTabWell(false)
 					->AddTab(SkeletalMeshEditorTabs::MorphTargetsTab, ETabState::OpenedTab)
 					->AddTab(SkeletalMeshEditorTabs::DetailsTab, ETabState::ClosedTab)
+	                ->AddTab(SkeletalMeshEditorTabs::ToolboxDetailsTab, ETabState::ClosedTab)
 					->AddTab(SkeletalMeshEditorTabs::AdvancedPreviewTab, ETabState::OpenedTab)
 					->SetForegroundTab(SkeletalMeshEditorTabs::MorphTargetsTab)
 				)
diff --git a/Engine/Source/Editor/SkeletalMeshEditor/SkeletalMeshEditor.Build.cs b/Engine/Source/Editor/SkeletalMeshEditor/SkeletalMeshEditor.Build.cs
--- a/Engine/Source/Editor/SkeletalMeshEditor/SkeletalMeshEditor.Build.cs
+++ b/Engine/Source/Editor/SkeletalMeshEditor/SkeletalMeshEditor.Build.cs
@@ -38,6 +38,7 @@ public class SkeletalMeshEditor : ModuleRules
 				"SkeletalMeshUtilitiesCommon",
 				"ToolMenus",
 				"EditorSubsystem",
+				"StatusBar",
 			}
 		);
 
