Cloth - Fix rebinding after undo.

diff --git a/Engine/Source/Editor/Persona/Private/PersonaMeshDetails.cpp b/Engine/Source/Editor/Persona/Private/PersonaMeshDetails.cpp
--- a/Engine/Source/Editor/Persona/Private/PersonaMeshDetails.cpp
+++ b/Engine/Source/Editor/Persona/Private/PersonaMeshDetails.cpp
@@ -6071,9 +6071,12 @@ void FPersonaMeshDetails::OnClothingSelectionChanged(TSharedPtr<FClothingEntry>
 		};
 		if (UClothingAssetCommon* ClothingAsset = Cast<UClothingAssetCommon>(InNewEntry->Asset.Get()))
 		{
+			ClothingAsset->Modify();
+
 			// Look for a currently bound asset an unbind it if necessary first
 			if (UClothingAssetBase* CurrentAsset = Mesh->GetSectionClothingAsset(InLodIdx, InSectionIdx))
 			{
+				CurrentAsset->Modify();
 				CurrentAsset->UnbindFromSkeletalMesh(Mesh, InLodIdx);
 				ClearOriginalSectionUserData();
 			}
@@ -6102,6 +6105,7 @@ void FPersonaMeshDetails::OnClothingSelectionChanged(TSharedPtr<FClothingEntry>
 			//User set none, so unbind anything that is bind
 			if (UClothingAssetBase* CurrentAsset = Mesh->GetSectionClothingAsset(InLodIdx, InSectionIdx))
 			{
+				CurrentAsset->Modify();
 				CurrentAsset->UnbindFromSkeletalMesh(Mesh, InLodIdx);
 				ClearOriginalSectionUserData();
 			}
diff --git a/Engine/Source/Editor/SkeletalMeshEditor/Private/SkeletalMeshEditor.cpp b/Engine/Source/Editor/SkeletalMeshEditor/Private/SkeletalMeshEditor.cpp
--- a/Engine/Source/Editor/SkeletalMeshEditor/Private/SkeletalMeshEditor.cpp
+++ b/Engine/Source/Editor/SkeletalMeshEditor/Private/SkeletalMeshEditor.cpp
@@ -1111,9 +1111,12 @@ void FSkeletalMeshEditor::ApplyClothing(UClothingAssetBase* InAsset, int32 InLod
 		};
 		if (UClothingAssetCommon* ClothingAsset = Cast<UClothingAssetCommon>(InAsset))
 		{
+			ClothingAsset->Modify();
+
 			// Look for a currently bound asset an unbind it if necessary first
 			if (UClothingAssetBase* CurrentAsset = Mesh->GetSectionClothingAsset(InLodIndex, InSectionIndex))
 			{
+				CurrentAsset->Modify();
 				CurrentAsset->UnbindFromSkeletalMesh(Mesh, InLodIndex);
 				ClearOriginalSectionUserData();
 			}
@@ -1133,6 +1136,7 @@ void FSkeletalMeshEditor::ApplyClothing(UClothingAssetBase* InAsset, int32 InLod
 			//User set none, so unbind anything that is bind
 			if (UClothingAssetBase* CurrentAsset = Mesh->GetSectionClothingAsset(InLodIndex, InSectionIndex))
 			{
+				CurrentAsset->Modify();
 				CurrentAsset->UnbindFromSkeletalMesh(Mesh, InLodIndex);
 				ClearOriginalSectionUserData();
 			}
