Add opt-in support for {Add,Remove}ViewportOverlayWidget to FAssetEditorToolkit derived editors.
Add implementation for it in the skeletal mesh editor.

diff --git a/Engine/Source/Editor/EditorFramework/Public/Toolkits/IToolkitHost.h b/Engine/Source/Editor/EditorFramework/Public/Toolkits/IToolkitHost.h
--- a/Engine/Source/Editor/EditorFramework/Public/Toolkits/IToolkitHost.h
+++ b/Engine/Source/Editor/EditorFramework/Public/Toolkits/IToolkitHost.h
@@ -47,8 +47,8 @@ public:
 
 	/** Allows Toolkits to push widgets to the viewport.  Passing a nullptr for the Viewport will add  or 
 	    remove the OverlayWidget to or from the ActiveViewport */
-	virtual void AddViewportOverlayWidget(TSharedRef<SWidget>, TSharedPtr<IAssetViewport> InViewport = nullptr)  {}
-	virtual void RemoveViewportOverlayWidget(TSharedRef<SWidget>, TSharedPtr<IAssetViewport> InViewport = nullptr) {}
+	virtual void AddViewportOverlayWidget(TSharedRef<SWidget> InOverlaidWidget, TSharedPtr<IAssetViewport> InViewport = nullptr)  {}
+	virtual void RemoveViewportOverlayWidget(TSharedRef<SWidget> InOverlaidWidget, TSharedPtr<IAssetViewport> InViewport = nullptr) {}
 
 	/** Gets a multicast delegate which is executed whenever the toolkit host's active viewport changes. */
 	virtual FOnActiveViewportChanged& OnActiveViewportChanged() = 0;
diff --git a/Engine/Source/Editor/Persona/Private/PersonaAssetEditorToolkit.cpp b/Engine/Source/Editor/Persona/Private/PersonaAssetEditorToolkit.cpp
--- a/Engine/Source/Editor/Persona/Private/PersonaAssetEditorToolkit.cpp
+++ b/Engine/Source/Editor/Persona/Private/PersonaAssetEditorToolkit.cpp
@@ -9,9 +9,5 @@
 
 void FPersonaAssetEditorToolkit::CreateEditorModeManager()
 {
-	TSharedPtr<FAssetEditorModeManager> ModeManager = MakeShareable(FModuleManager::LoadModuleChecked<FPersonaModule>("Persona").CreatePersonaEditorModeManager());
-	
-	PRAGMA_DISABLE_DEPRECATION_WARNINGS
-	SetAssetEditorModeManager(ModeManager.Get());
-	PRAGMA_ENABLE_DEPRECATION_WARNINGS
+	EditorModeManager = MakeShareable(FModuleManager::LoadModuleChecked<FPersonaModule>("Persona").CreatePersonaEditorModeManager());
 }
diff --git a/Engine/Source/Editor/Persona/Private/SAnimationEditorViewport.cpp b/Engine/Source/Editor/Persona/Private/SAnimationEditorViewport.cpp
--- a/Engine/Source/Editor/Persona/Private/SAnimationEditorViewport.cpp
+++ b/Engine/Source/Editor/Persona/Private/SAnimationEditorViewport.cpp
@@ -411,11 +411,6 @@ TWeakPtr<SWidget> SAnimationEditorViewportTabBody::AddNotification(TAttribute<EM
 	return ContainingWidget;
 }
 
-void SAnimationEditorViewportTabBody::AddToolbarExtender(FName MenuToExtend, FMenuExtensionDelegate MenuBuilderDelegate)
-{
-	return ViewportWidget->ViewportToolbar->AddMenuExtender(MenuToExtend, MenuBuilderDelegate);
-}
-
 void SAnimationEditorViewportTabBody::RemoveNotification(const TWeakPtr<SWidget>& InContainingWidget)
 {
 	if(InContainingWidget.IsValid())
@@ -424,6 +419,25 @@ void SAnimationEditorViewportTabBody::RemoveNotification(const TWeakPtr<SWidget>
 	}
 }
 
+
+void SAnimationEditorViewportTabBody::AddToolbarExtender(FName MenuToExtend, FMenuExtensionDelegate MenuBuilderDelegate)
+{
+	return ViewportWidget->ViewportToolbar->AddMenuExtender(MenuToExtend, MenuBuilderDelegate);
+}
+
+void SAnimationEditorViewportTabBody::AddOverlayWidget(TSharedRef<SWidget> InOverlaidWidget)
+{
+	ViewportWidget->ViewportOverlay->AddSlot()
+	[
+		InOverlaidWidget
+	];
+}
+
+void SAnimationEditorViewportTabBody::RemoveOverlayWidget(TSharedRef<SWidget> InOverlaidWidget)
+{
+	ViewportWidget->ViewportOverlay->RemoveSlot(InOverlaidWidget);
+}
+
 void SAnimationEditorViewportTabBody::RefreshViewport()
 {
 	LevelViewportClient->Invalidate();
diff --git a/Engine/Source/Editor/Persona/Private/SAnimationEditorViewport.h b/Engine/Source/Editor/Persona/Private/SAnimationEditorViewport.h
--- a/Engine/Source/Editor/Persona/Private/SAnimationEditorViewport.h
+++ b/Engine/Source/Editor/Persona/Private/SAnimationEditorViewport.h
@@ -220,10 +220,13 @@ public:
 	virtual FEditorViewportClient& GetViewportClient() const override;
 	virtual TSharedRef<IPinnedCommandList> GetPinnedCommandList() const override;
 	virtual TWeakPtr<SWidget> AddNotification(TAttribute<EMessageSeverity::Type> InSeverity, TAttribute<bool> InCanBeDismissed, const TSharedRef<SWidget>& InNotificationWidget) override;
+	virtual void RemoveNotification(const TWeakPtr<SWidget>& InContainingWidget) override;
 	virtual void AddToolbarExtender(FName MenuToExtend, FMenuExtensionDelegate MenuBuilderDelegate) override;
 	virtual FPersonaViewportKeyDownDelegate& GetKeyDownDelegate() override { return OnKeyDownDelegate; }
-	virtual void RemoveNotification(const TWeakPtr<SWidget>& InContainingWidget) override;
+	virtual void AddOverlayWidget( TSharedRef<SWidget> InOverlaidWidget ) override;
+	virtual void RemoveOverlayWidget( TSharedRef<SWidget> InOverlaidWidget ) override;
 
+	
 	/** SWidget interface */
 	virtual FReply OnKeyDown(const FGeometry& MyGeometry, const FKeyEvent& InKeyEvent) override;
 
diff --git a/Engine/Source/Editor/Persona/Public/IPersonaViewport.h b/Engine/Source/Editor/Persona/Public/IPersonaViewport.h
--- a/Engine/Source/Editor/Persona/Public/IPersonaViewport.h
+++ b/Engine/Source/Editor/Persona/Public/IPersonaViewport.h
@@ -41,6 +41,12 @@ public:
 	 * @return the widget containing the notification
 	 */
 	virtual TWeakPtr<SWidget> AddNotification(TAttribute<EMessageSeverity::Type> InSeverity, TAttribute<bool> InCanBeDismissed, const TSharedRef<SWidget>& InNotificationWidget) = 0;
+	
+	/** 
+	* Remove a notification widget 
+	* @param	InContainingWidget		The containing widget returned from AddNotification()
+	*/
+	virtual void RemoveNotification(const TWeakPtr<SWidget>& InContainingWidget) = 0;
 
 	/** 
 	 * Adds an extender to the viewport's toolbar
@@ -54,9 +60,9 @@ public:
 	 */
 	virtual FPersonaViewportKeyDownDelegate& GetKeyDownDelegate() = 0;
 
-	/** 
-	 * Remove a notification widget 
-	 * @param	InContainingWidget		The containing widget returned from AddNotification()
-	 */
-	virtual void RemoveNotification(const TWeakPtr<SWidget>& InContainingWidget) = 0;
+	/** Overlay a widget over the whole viewport */
+	virtual void AddOverlayWidget( TSharedRef<SWidget> InOverlaidWidget ) = 0;
+
+	/** Remove an overlay widget from the viewport */
+	virtual void RemoveOverlayWidget( TSharedRef<SWidget> InOverlaidWidget ) = 0;
 };
diff --git a/Engine/Source/Editor/SkeletalMeshEditor/Private/SkeletalMeshEditor.cpp b/Engine/Source/Editor/SkeletalMeshEditor/Private/SkeletalMeshEditor.cpp
--- a/Engine/Source/Editor/SkeletalMeshEditor/Private/SkeletalMeshEditor.cpp
+++ b/Engine/Source/Editor/SkeletalMeshEditor/Private/SkeletalMeshEditor.cpp
@@ -29,6 +29,7 @@
 #include "IDetailsView.h"
 #include "IPersonaPreviewScene.h"
 #include "IPersonaToolkit.h"
+#include "IPersonaViewport.h"
 #include "ISkeletalMeshEditorModule.h"
 #include "ISkeletonEditorModule.h"
 #include "ISkeletonTree.h"
@@ -703,6 +704,24 @@ void FSkeletalMeshEditor::OnToolkitHostingFinished(const TSharedRef<IToolkit>& T
 	}
 }
 
+
+void FSkeletalMeshEditor::AddViewportOverlayWidget(TSharedRef<SWidget> InOverlaidWidget)
+{
+	if (Viewport.IsValid())
+	{
+		Viewport->AddOverlayWidget(InOverlaidWidget);
+	}	
+}
+
+
+void FSkeletalMeshEditor::RemoveViewportOverlayWidget(TSharedRef<SWidget> InOverlaidWidget)
+{
+	if (Viewport.IsValid())
+	{
+		Viewport->RemoveOverlayWidget(InOverlaidWidget);
+	}	
+}
+
 bool FSkeletalMeshEditor::ProcessCommandBindings(const FKeyEvent& InKeyEvent) const
 {
 	if (HostedToolkit.IsValid() && HostedToolkit->ProcessCommandBindings(InKeyEvent))
@@ -1244,6 +1263,11 @@ void FSkeletalMeshEditor::HandleMeshDetailsCreated(const TSharedRef<IDetailsView
 	PersonaModule.CustomizeMeshDetails(InDetailsView, GetPersonaToolkit());
 }
 
+void FSkeletalMeshEditor::HandleViewportCreated(const TSharedRef<IPersonaViewport>& InViewport)
+{
+	Viewport = InViewport;
+}
+
 UObject* FSkeletalMeshEditor::HandleGetAsset()
 {
 	return GetEditingObject();
diff --git a/Engine/Source/Editor/SkeletalMeshEditor/Private/SkeletalMeshEditor.h b/Engine/Source/Editor/SkeletalMeshEditor/Private/SkeletalMeshEditor.h
--- a/Engine/Source/Editor/SkeletalMeshEditor/Private/SkeletalMeshEditor.h
+++ b/Engine/Source/Editor/SkeletalMeshEditor/Private/SkeletalMeshEditor.h
@@ -71,6 +71,9 @@ public:
 	virtual void OnToolkitHostingStarted(const TSharedRef<class IToolkit>& Toolkit) override;
 	virtual void OnToolkitHostingFinished(const TSharedRef<class IToolkit>& Toolkit) override;
 
+	virtual void AddViewportOverlayWidget(TSharedRef<SWidget> InOverlaidWidget) override;
+	virtual void RemoveViewportOverlayWidget(TSharedRef<SWidget> InOverlaidWidget) override;
+	
 	/** FBaseToolkit overrides */
 	virtual bool ProcessCommandBindings(const FKeyEvent& InKeyEvent) const override;
 	
@@ -107,6 +110,8 @@ public:
 
 	void HandleMeshDetailsCreated(const TSharedRef<class IDetailsView>& InDetailsView);
 
+	void HandleViewportCreated(const TSharedRef<class IPersonaViewport>& InViewport);
+
 	UObject* HandleGetAsset();
 
 	// Returns the currently hosted toolkit. Can be invalid if no toolkit is being hosted.
diff --git a/Engine/Source/Editor/SkeletalMeshEditor/Private/SkeletalMeshEditorMode.cpp b/Engine/Source/Editor/SkeletalMeshEditor/Private/SkeletalMeshEditorMode.cpp
--- a/Engine/Source/Editor/SkeletalMeshEditor/Private/SkeletalMeshEditorMode.cpp
+++ b/Engine/Source/Editor/SkeletalMeshEditor/Private/SkeletalMeshEditorMode.cpp
@@ -26,6 +26,7 @@ FSkeletalMeshEditorMode::FSkeletalMeshEditorMode(TSharedRef<FWorkflowCentricAppl
 
 	FPersonaViewportArgs ViewportArgs(SkeletalMeshEditor->GetPersonaToolkit()->GetPreviewScene());
 	ViewportArgs.ContextName = TEXT("SkeletalMeshEditor.Viewport");
+	ViewportArgs.OnViewportCreated = FOnViewportCreated::CreateSP(SkeletalMeshEditor, &FSkeletalMeshEditor::HandleViewportCreated);
 
 	PersonaModule.RegisterPersonaViewportTabFactories(TabFactories, InHostingApp, ViewportArgs);
 
diff --git a/Engine/Source/Editor/UnrealEd/Private/Toolkits/SStandaloneAssetEditorToolkitHost.cpp b/Engine/Source/Editor/UnrealEd/Private/Toolkits/SStandaloneAssetEditorToolkitHost.cpp
--- a/Engine/Source/Editor/UnrealEd/Private/Toolkits/SStandaloneAssetEditorToolkitHost.cpp
+++ b/Engine/Source/Editor/UnrealEd/Private/Toolkits/SStandaloneAssetEditorToolkitHost.cpp
@@ -342,6 +342,22 @@ UTypedElementCommonActions* SStandaloneAssetEditorToolkitHost::GetCommonActions(
 	return nullptr;
 }
 
+void SStandaloneAssetEditorToolkitHost::AddViewportOverlayWidget(TSharedRef<SWidget> InOverlaidWidget, TSharedPtr<IAssetViewport> InViewport)
+{
+	if (HostedAssetEditorToolkit.IsValid())
+	{
+		HostedAssetEditorToolkit->AddViewportOverlayWidget(InOverlaidWidget);
+	}
+}
+
+void SStandaloneAssetEditorToolkitHost::RemoveViewportOverlayWidget(TSharedRef<SWidget> InOverlaidWidget, TSharedPtr<IAssetViewport> InViewport)
+{
+	if (HostedAssetEditorToolkit.IsValid())
+	{
+		HostedAssetEditorToolkit->RemoveViewportOverlayWidget(InOverlaidWidget);
+	}
+}
+
 
 FReply SStandaloneAssetEditorToolkitHost::OnKeyDown( const FGeometry& MyGeometry, const FKeyEvent& InKeyEvent )
 {
diff --git a/Engine/Source/Editor/UnrealEd/Private/Toolkits/SStandaloneAssetEditorToolkitHost.h b/Engine/Source/Editor/UnrealEd/Private/Toolkits/SStandaloneAssetEditorToolkitHost.h
--- a/Engine/Source/Editor/UnrealEd/Private/Toolkits/SStandaloneAssetEditorToolkitHost.h
+++ b/Engine/Source/Editor/UnrealEd/Private/Toolkits/SStandaloneAssetEditorToolkitHost.h
@@ -57,6 +57,8 @@ public:
 	virtual void OnToolkitHostingFinished( const TSharedRef< class IToolkit >& Toolkit ) override;
 	virtual UWorld* GetWorld() const override;
 	virtual UTypedElementCommonActions* GetCommonActions() const override;
+	virtual void AddViewportOverlayWidget(TSharedRef<SWidget> InOverlaidWidget, TSharedPtr<IAssetViewport> InViewport = nullptr) override;
+	virtual void RemoveViewportOverlayWidget(TSharedRef<SWidget> InOverlaidWidget, TSharedPtr<IAssetViewport> InViewport = nullptr) override;
 	virtual FOnActiveViewportChanged& OnActiveViewportChanged() override { return OnActiveViewportChangedDelegate; }
 	void CreateDefaultStandaloneMenuBar(UToolMenu* MenuBar);
 
diff --git a/Engine/Source/Editor/UnrealEd/Public/Toolkits/AssetEditorToolkit.h b/Engine/Source/Editor/UnrealEd/Public/Toolkits/AssetEditorToolkit.h
--- a/Engine/Source/Editor/UnrealEd/Public/Toolkits/AssetEditorToolkit.h
+++ b/Engine/Source/Editor/UnrealEd/Public/Toolkits/AssetEditorToolkit.h
@@ -215,6 +215,14 @@ public:
 	// Called when another toolkit (such as a ed mode toolkit) is no longer being hosted in this asset editor toolkit
 	virtual void OnToolkitHostingFinished(const TSharedRef<IToolkit>& Toolkit) {}
 
+	/* Called when a toolkit requests an overlay widget to be added to the viewport. Not relevant in the absence
+	 * of a viewport.
+	 */
+	virtual void AddViewportOverlayWidget(TSharedRef<SWidget> InViewportOverlayWidget) {}
+
+	/** Called when a toolkit requests the overlay widget to be removed. */
+	virtual void RemoveViewportOverlayWidget(TSharedRef<SWidget> InViewportOverlayWidget) {}
+	
 	/** Adds or removes extenders to the default menu or the toolbar menu this asset editor */
 	void AddMenuExtender(TSharedPtr<FExtender> Extender);
 	void RemoveMenuExtender(TSharedPtr<FExtender> Extender);
