Persona: Adapt skeletalmesh editor to capsule based selection

diff --git a/Engine/Source/Editor/Persona/Private/AnimationEditorViewportClient.cpp b/Engine/Source/Editor/Persona/Private/AnimationEditorViewportClient.cpp
--- a/Engine/Source/Editor/Persona/Private/AnimationEditorViewportClient.cpp
+++ b/Engine/Source/Editor/Persona/Private/AnimationEditorViewportClient.cpp
@@ -52,9 +52,6 @@ namespace EAnimationPlaybackSpeeds
 	float Values[EAnimationPlaybackSpeeds::NumPlaybackSpeeds] = { 0.1f, 0.25f, 0.5f, 1.0f, 2.0f, 5.0f, 10.0f };
 }
 
-IMPLEMENT_HIT_PROXY( HPersonaSocketProxy, HHitProxy );
-IMPLEMENT_HIT_PROXY( HPersonaBoneProxy, HHitProxy );
-
 #define LOCTEXT_NAMESPACE "FAnimationViewportClient"
 
 /////////////////////////////////////////////////////////////////////////
@@ -1399,9 +1396,7 @@ void FAnimationViewportClient::DrawBones(const TArray<FBoneIndexType>& RequiredB
 			// clamp by bound, we don't want too long or big
 			const float Radius = FMath::Clamp(BoneLength * 0.05f, 0.1f, MaxDrawRadius) * InBoneRadius;
 			//Render Sphere for bone end point and a cone between it and its parent.
-			PDI->SetHitProxy(new HPersonaBoneProxy(RefSkeleton.GetBoneName(BoneIndex)));
 			SkeletalDebugRendering::DrawWireBone(PDI, Start, End, LineColor, SDPG_Foreground, Radius);
-			PDI->SetHitProxy(NULL);
 
 			// draw gizmo
 			if ((GetLocalAxesMode() == ELocalAxesMode::All) ||
@@ -1529,9 +1524,7 @@ void FAnimationViewportClient::DrawSockets(const UDebugSkelMeshComponent* InPrev
 			float Angle = FMath::RadiansToDegrees(FMath::Atan(SphereRadius / ConeLength));
 
 			//Render Sphere for bone end point and a cone between it and its parent.
-			PDI->SetHitProxy( new HPersonaBoneProxy( Socket->BoneName ) );
 			PDI->DrawLine( Start, End, SocketColor, SDPG_Foreground );
-			PDI->SetHitProxy( NULL );
 			
 			// draw gizmo
 			if( (LocalAxesMode == ELocalAxesMode::All) || bSelectedSocket )
@@ -1539,9 +1532,7 @@ void FAnimationViewportClient::DrawSockets(const UDebugSkelMeshComponent* InPrev
 				FMatrix SocketMatrix;
 				Socket->GetSocketMatrix( SocketMatrix, InPreviewMeshComponent);
 
-				PDI->SetHitProxy( new HPersonaSocketProxy( FSelectedSocketInfo( Socket, bUseSkeletonSocketColor ) ) );
 				DrawWireDiamond( PDI, SocketMatrix, 2.f, SocketColor, SDPG_Foreground );
-				PDI->SetHitProxy( NULL );
 				
 				SkeletalDebugRendering::DrawAxes(PDI, FTransform(SocketMatrix), SDPG_Foreground);
 			}
diff --git a/Engine/Source/Editor/Persona/Private/AnimationEditorViewportClient.h b/Engine/Source/Editor/Persona/Private/AnimationEditorViewportClient.h
--- a/Engine/Source/Editor/Persona/Private/AnimationEditorViewportClient.h
+++ b/Engine/Source/Editor/Persona/Private/AnimationEditorViewportClient.h
@@ -80,34 +80,6 @@ namespace EAnimationPlaybackSpeeds
 	extern float Values[NumPlaybackSpeeds];
 };
 
-/**
-* A hit proxy class for sockets in the Persona viewport.
-*/
-struct HPersonaSocketProxy : public HHitProxy
-{
-	DECLARE_HIT_PROXY();
-
-	FSelectedSocketInfo SocketInfo;
-
-	explicit HPersonaSocketProxy(FSelectedSocketInfo InSocketInfo)
-		: SocketInfo(InSocketInfo)
-	{}
-};
-
-/**
-* A hit proxy class for sockets in the Persona viewport.
-*/
-struct HPersonaBoneProxy : public HHitProxy
-{
-	DECLARE_HIT_PROXY();
-
-	FName BoneName;
-
-	explicit HPersonaBoneProxy(const FName& InBoneName)
-		: BoneName(InBoneName)
-	{}
-};
-
 /////////////////////////////////////////////////////////////////////////
 // FAnimationViewportClient
 
diff --git a/Engine/Source/Editor/Persona/Private/EditModes/SkeletonSelectionEditMode.cpp b/Engine/Source/Editor/Persona/Private/EditModes/SkeletonSelectionEditMode.cpp
--- a/Engine/Source/Editor/Persona/Private/EditModes/SkeletonSelectionEditMode.cpp
+++ b/Engine/Source/Editor/Persona/Private/EditModes/SkeletonSelectionEditMode.cpp
@@ -529,18 +529,10 @@ bool FSkeletonSelectionEditMode::HandleClick(FEditorViewportClient* InViewportCl
 			MeshComponent->SetSelectedEditorSection(INDEX_NONE);
 		}
 
-		if ( HitProxy->IsA( HPersonaSocketProxy::StaticGetType() ) )
+		if ( HitProxy->IsA( HPersonaSelectionHitProxy::StaticGetType() ) )
 		{
-			// Tell the preview scene that the socket has been selected - this will sort out the skeleton tree, etc.
 			GetAnimPreviewScene().DeselectAll();
-			GetAnimPreviewScene().SetSelectedSocket(static_cast<HPersonaSocketProxy*>(HitProxy)->SocketInfo);
-			bHandled = true;
-		}
-		else if ( HitProxy->IsA( HPersonaBoneProxy::StaticGetType() ) )
-		{			
-			// Tell the preview scene that the bone has been selected - this will sort out the skeleton tree, etc.
-			GetAnimPreviewScene().DeselectAll();
-			GetAnimPreviewScene().SetSelectedBone(static_cast<HPersonaBoneProxy*>(HitProxy)->BoneName, ESelectInfo::OnMouseClick);
+			static_cast<HPersonaSelectionHitProxy*>(HitProxy)->BroadcastClicked();
 			bHandled = true;
 		}
 		else if ( HitProxy->IsA( HActor::StaticGetType() ) && bSelectingSections)
diff --git a/Engine/Source/Editor/SkeletalMeshEditor/Private/SkeletalMeshEditor.cpp b/Engine/Source/Editor/SkeletalMeshEditor/Private/SkeletalMeshEditor.cpp
--- a/Engine/Source/Editor/SkeletalMeshEditor/Private/SkeletalMeshEditor.cpp
+++ b/Engine/Source/Editor/SkeletalMeshEditor/Private/SkeletalMeshEditor.cpp
@@ -47,6 +47,7 @@
 #include "Settings/EditorExperimentalSettings.h"
 #include "SkeletalMeshToolMenuContext.h"
 #include "ToolMenus.h"
+#include "Engine/SkeletalMeshSocket.h"
 #include "Widgets/Docking/SDockTab.h"
 #include "Widgets/Layout/SBorder.h"
 #include "Widgets/Layout/SBox.h"
@@ -195,6 +196,50 @@ bool FSkeletalMeshEditor::OnRequestClose()
 	return bAllowClose;
 }
 
+void FSkeletalMeshEditor::UpdateCapsules()
+{
+	UPersonaSelectionComponent* SelectionComponent = PersonaToolkit->GetPreviewScene()->GetSelectionComponent();
+	check(SelectionComponent != nullptr);
+	
+	SelectionComponent->Reset();
+	CapsuleInfos.Reset();
+	
+	// setup the capsules for selecting bones / sockets in the viewport
+	const FReferenceSkeleton& ReferenceSkeleton = PersonaToolkit->GetSkeleton()->GetReferenceSkeleton();
+	for(int32 Index = 0; Index < ReferenceSkeleton.GetNum(); Index++)
+	{
+		FCapsuleInfo Info;
+		Info.Type = ECapsuleType::Bone;
+		Info.Name = ReferenceSkeleton.GetBoneName(Index);
+		Info.Index = Index;
+		CapsuleInfos.Add(Info);
+
+		SelectionComponent->Add();
+	}
+
+	for(int32 Index = 0; Index < PersonaToolkit->GetSkeleton()->GetVirtualBones().Num(); Index++)
+	{
+		FCapsuleInfo Info;
+		Info.Type = ECapsuleType::VirtualBone;
+		Info.Name = PersonaToolkit->GetSkeleton()->GetVirtualBones()[Index].VirtualBoneName;
+		Info.Index = Index;
+		CapsuleInfos.Add(Info);
+
+		SelectionComponent->Add();
+	}
+
+	for(int32 Index = 0; Index < PersonaToolkit->GetSkeleton()->Sockets.Num(); Index++)
+	{
+		FCapsuleInfo Info;
+		Info.Type = ECapsuleType::Socket;
+		Info.Name = PersonaToolkit->GetSkeleton()->Sockets[Index]->GetFName();
+		Info.Index = Index;
+		CapsuleInfos.Add(Info);
+
+		SelectionComponent->Add();
+	}
+}
+
 void FSkeletalMeshEditor::RegisterTabSpawners(const TSharedRef<class FTabManager>& InTabManager)
 {
 	WorkspaceMenuCategory = InTabManager->AddLocalWorkspaceMenuCategory(LOCTEXT("WorkspaceMenu_SkeletalMeshEditor", "Skeletal Mesh Editor"));
@@ -257,6 +302,95 @@ void FSkeletalMeshEditor::InitSkeletalMeshEditor(const EToolkitMode::Type Mode,
 	// Set up mesh click selection
 	PreviewScene->RegisterOnMeshClick(FOnMeshClick::CreateSP(this, &FSkeletalMeshEditor::HandleMeshClick));
 	PreviewScene->SetAllowMeshHitProxies(GetDefault<UPersonaOptions>()->bAllowMeshSectionSelection);
+	if(PreviewScene->GetPreviewMeshComponent())
+	{
+		PreviewScene->GetPreviewMeshComponent()->bSelectable = PreviewScene->AllowMeshHitProxies();
+		PreviewScene->GetPreviewMeshComponent()->MarkRenderStateDirty();
+	}
+
+	UpdateCapsules();
+	
+	// setup the lambda to fill the capsules
+	UPersonaSelectionComponent* SelectionComponent = PreviewScene->GetSelectionComponent();
+	check(SelectionComponent != nullptr);
+
+	// make sure the selection component ticks after the skeletal mesh comp
+	SelectionComponent->AddTickPrerequisiteComponent(PreviewScene->GetPreviewMeshComponent());
+
+	SelectionComponent->OnUpdateCapsules().BindLambda(
+		[this, PreviewScene](UPersonaSelectionComponent*, const TArray<int32>& InIndices, TArray<FPersonaSelectionCapsule>& OutCapsules)
+		{
+			const FReferenceSkeleton& ReferenceSkeleton = PreviewScene->GetPreviewMeshComponent()->GetReferenceSkeleton();
+			TArray<FTransform> ComponentSpaceTransforms = PreviewScene->GetPreviewMeshComponent()->GetComponentSpaceTransforms();
+			
+			for(int32 Index = 0; Index < InIndices.Num(); Index++)
+			{
+				const int32 CapsuleIndex = InIndices[Index];
+				check(CapsuleInfos.IsValidIndex(CapsuleIndex));
+
+				const FCapsuleInfo& CapsuleInfo = CapsuleInfos[CapsuleIndex];
+				switch(CapsuleInfo.Type)
+				{
+					case ECapsuleType::Bone:
+					{
+						OutCapsules[CapsuleIndex].Name = CapsuleInfo.Name;
+
+						const FVector& BonePosition = ComponentSpaceTransforms[CapsuleInfo.Index].GetLocation();
+
+						FVector Start, End;
+							
+						const int32 ParentIndex = ReferenceSkeleton.GetParentIndex(CapsuleInfo.Index);
+						if(ParentIndex != INDEX_NONE)
+						{
+							Start = ComponentSpaceTransforms[ParentIndex].GetLocation();
+							End = BonePosition;
+						}
+						else
+						{
+							Start = FVector::ZeroVector;
+							End = BonePosition;
+						}
+
+						UPersonaSelectionComponent::ComputeCapsuleFromBonePositions(Start, End, 100.f, 1.f, OutCapsules[CapsuleIndex]);
+						break;
+					}
+					default:
+					{
+						break;
+					}
+				}
+			}
+		});
+
+	SelectionComponent->OnClicked().BindLambda(
+		[PreviewScene](UPersonaSelectionComponent*, int32 InIndex, const FPersonaSelectionCapsule& InCapsule)
+		{
+			PreviewScene->DeselectAll();
+
+			bool bFoundSocket = false;
+
+			if(USkeleton* Skeleton = PreviewScene->GetPersonaToolkit()->GetSkeleton())
+			{
+				for(USkeletalMeshSocket* Socket : Skeleton->Sockets)
+				{
+					if(Socket->GetFName() == InCapsule.Name)
+					{
+						bFoundSocket = true;
+
+						FSelectedSocketInfo SocketInfo;
+						SocketInfo.Socket = Socket;
+						SocketInfo.bSocketIsOnSkeleton = !SocketInfo.Socket->GetOuter()->IsA<USkeletalMesh>();
+						PreviewScene->SetSelectedSocket(SocketInfo);
+					}
+				}
+			}
+						
+			if(!bFoundSocket)
+			{
+				PreviewScene->SetSelectedBone(InCapsule.Name, ESelectInfo::OnMouseClick);
+			}
+
+		});
 
 	// Make sure we get told when the editor mode changes so we can switch to the appropriate tab
 	// if there's a toolbox available.
diff --git a/Engine/Source/Editor/SkeletalMeshEditor/Private/SkeletalMeshEditor.h b/Engine/Source/Editor/SkeletalMeshEditor/Private/SkeletalMeshEditor.h
--- a/Engine/Source/Editor/SkeletalMeshEditor/Private/SkeletalMeshEditor.h
+++ b/Engine/Source/Editor/SkeletalMeshEditor/Private/SkeletalMeshEditor.h
@@ -212,4 +212,22 @@ private:
 
 	// The toolbox widget
 	TSharedPtr<SSkeletalMeshEditorToolbox> ToolboxWidget;
+
+	enum ECapsuleType
+	{
+		Bone,
+		VirtualBone,
+		Socket
+	};
+	
+	struct FCapsuleInfo
+	{
+		ECapsuleType Type;
+		FName Name;
+		int32 Index;
+	};
+
+	TArray<FCapsuleInfo> CapsuleInfos;
+
+	void UpdateCapsules();
 };
