SkeletalMeshEditor: Guard against invalid indices

diff --git a/Engine/Source/Editor/SkeletalMeshEditor/Private/SkeletalMeshEditor.cpp b/Engine/Source/Editor/SkeletalMeshEditor/Private/SkeletalMeshEditor.cpp
--- a/Engine/Source/Editor/SkeletalMeshEditor/Private/SkeletalMeshEditor.cpp
+++ b/Engine/Source/Editor/SkeletalMeshEditor/Private/SkeletalMeshEditor.cpp
@@ -334,13 +334,17 @@ void FSkeletalMeshEditor::InitSkeletalMeshEditor(const EToolkitMode::Type Mode,
 					case ECapsuleType::Bone:
 					{
 						OutCapsules[CapsuleIndex].Name = CapsuleInfo.Name;
+						if(!ComponentSpaceTransforms.IsValidIndex(CapsuleInfo.Index))
+						{
+							break;
+						}
 
 						const FVector& BonePosition = ComponentSpaceTransforms[CapsuleInfo.Index].GetLocation();
 
 						FVector Start, End;
 							
 						const int32 ParentIndex = ReferenceSkeleton.GetParentIndex(CapsuleInfo.Index);
-						if(ParentIndex != INDEX_NONE)
+						if(ComponentSpaceTransforms.IsValidIndex(ParentIndex))
 						{
 							Start = ComponentSpaceTransforms[ParentIndex].GetLocation();
 							End = BonePosition;
