SkeletalMeshEditor: Clear dependency on selection actor

diff --git a/Engine/Source/Editor/SkeletalMeshEditor/Private/SkeletalMeshEditor.cpp b/Engine/Source/Editor/SkeletalMeshEditor/Private/SkeletalMeshEditor.cpp
--- a/Engine/Source/Editor/SkeletalMeshEditor/Private/SkeletalMeshEditor.cpp
+++ b/Engine/Source/Editor/SkeletalMeshEditor/Private/SkeletalMeshEditor.cpp
@@ -89,6 +89,12 @@ FSkeletalMeshEditor::~FSkeletalMeshEditor()
 	{
 		Editor->UnregisterForUndo(this);
 	}
+
+	if(PersonaToolkit.IsValid())
+	{
+		UPersonaSelectionComponent* SelectionComponent = PersonaToolkit->GetPreviewScene()->GetSelectionComponent();
+		SelectionComponent->GetWorld()->RemoveActor(SelectionComponent->GetOwner(), true);
+	}
 }
 
 bool IsReductionParentBaseLODUseSkeletalMeshBuildWorkflow(USkeletalMesh* SkeletalMesh, int32 TestLODIndex)
