Adding Bake Out Materials back to the Skeletal Mesh Editor (on the Asset Menu instead of the toolbar)

diff --git a/Engine/Source/Editor/SkeletalMeshEditor/Private/SkeletalMeshEditor.cpp b/Engine/Source/Editor/SkeletalMeshEditor/Private/SkeletalMeshEditor.cpp
--- a/Engine/Source/Editor/SkeletalMeshEditor/Private/SkeletalMeshEditor.cpp
+++ b/Engine/Source/Editor/SkeletalMeshEditor/Private/SkeletalMeshEditor.cpp
@@ -54,6 +54,8 @@
 #include "Interfaces/ITargetPlatform.h"
 #include "Interfaces/ITargetPlatformManagerModule.h"
 #include "Misc/CoreMisc.h"
+#include "Toolkits/AssetEditorToolkitMenuContext.h"
+#include "MeshMergeModule.h"
 
 const FName SkeletalMeshEditorAppIdentifier = FName(TEXT("SkeletalMeshEditorApp"));
 
@@ -449,6 +451,11 @@ void FSkeletalMeshEditor::BindCommands()
 
 	ToolkitCommands->MapAction(FPersonaCommonCommands::Get().TogglePlay,
 		FExecuteAction::CreateRaw(&GetPersonaToolkit()->GetPreviewScene().Get(), &IPersonaPreviewScene::TogglePlayback));
+
+	// Bake Materials
+	ToolkitCommands->MapAction(FSkeletalMeshEditorCommands::Get().BakeMaterials,
+		FExecuteAction::CreateSP(this, &FSkeletalMeshEditor::BakeMaterials),
+		FCanExecuteAction());
 }
 
 TSharedPtr<FSkeletalMeshEditor> FSkeletalMeshEditor::GetSkeletalMeshEditor(const FToolMenuContext& InMenuContext)
@@ -1348,6 +1355,27 @@ void FSkeletalMeshEditor::ExtendMenu()
 
 	ISkeletalMeshEditorModule& SkeletalMeshEditorModule = FModuleManager::GetModuleChecked<ISkeletalMeshEditorModule>("SkeletalMeshEditor");
 	AddMenuExtender(SkeletalMeshEditorModule.GetMenuExtensibilityManager()->GetAllExtenders(GetToolkitCommands(), GetEditingObjects()));
+
+	UToolMenu* AssetMenu = UToolMenus::Get()->ExtendMenu("MainFrame.MainMenu.Asset");
+	FToolMenuSection& AssetSection = AssetMenu->FindOrAddSection("AssetEditorActions");
+	FToolMenuEntry& Entry = AssetSection.AddDynamicEntry("AssetManagerEditorSkeletalMeshCommands", FNewToolMenuSectionDelegate::CreateLambda([this](FToolMenuSection& InSection)
+		{
+			UAssetEditorToolkitMenuContext* MenuContext = InSection.FindContext<UAssetEditorToolkitMenuContext>();
+			if (MenuContext && MenuContext->Toolkit.IsValid() && MenuContext->Toolkit.Pin()->GetToolkitFName() == GetToolkitFName())
+			{
+				InSection.AddMenuEntry(FSkeletalMeshEditorCommands::Get().BakeMaterials);
+			}
+		}
+	));
+}
+
+void FSkeletalMeshEditor::BakeMaterials()
+{
+	if (GetPersonaToolkit()->GetPreviewMeshComponent() != nullptr)
+	{
+		const IMeshMergeModule& Module = FModuleManager::Get().LoadModuleChecked<IMeshMergeModule>("MeshMergeUtilities");
+		Module.GetUtilities().BakeMaterialsForComponent(GetPersonaToolkit()->GetPreviewMeshComponent());
+	}
 }
 
 void FSkeletalMeshEditor::HandleObjectsSelected(const TArray<UObject*>& InObjects)
diff --git a/Engine/Source/Editor/SkeletalMeshEditor/Private/SkeletalMeshEditor.h b/Engine/Source/Editor/SkeletalMeshEditor/Private/SkeletalMeshEditor.h
--- a/Engine/Source/Editor/SkeletalMeshEditor/Private/SkeletalMeshEditor.h
+++ b/Engine/Source/Editor/SkeletalMeshEditor/Private/SkeletalMeshEditor.h
@@ -174,6 +174,7 @@ private:
 private:
 	void ExtendMenu();
 
+	void BakeMaterials();
 	void ExtendToolbar();
 
 	void BindCommands();
diff --git a/Engine/Source/Editor/SkeletalMeshEditor/Private/SkeletalMeshEditorCommands.cpp b/Engine/Source/Editor/SkeletalMeshEditor/Private/SkeletalMeshEditorCommands.cpp
--- a/Engine/Source/Editor/SkeletalMeshEditor/Private/SkeletalMeshEditorCommands.cpp
+++ b/Engine/Source/Editor/SkeletalMeshEditor/Private/SkeletalMeshEditorCommands.cpp
@@ -13,6 +13,8 @@ void FSkeletalMeshEditorCommands::RegisterCommands()
 	UI_COMMAND(ReimportAllMeshWithNewFile, "Reimport Base Mesh + LODs With New File", "Reimport the base mesh using a new source file and all the custom LODs (No new source file for LODs).", EUserInterfaceActionType::Button, FInputChord());
 
 	UI_COMMAND(MeshSectionSelection, "Section Selection", "Enables selecting Mesh Sections in the viewport (disables selecting bones using their physics shape).", EUserInterfaceActionType::ToggleButton, FInputChord());
+
+	UI_COMMAND(BakeMaterials, "Bake Out Materials", "Bake out Materials for given LOD(s).", EUserInterfaceActionType::Button, FInputChord());
 }
 
 #undef LOCTEXT_NAMESPACE
diff --git a/Engine/Source/Editor/SkeletalMeshEditor/Private/SkeletalMeshEditorCommands.h b/Engine/Source/Editor/SkeletalMeshEditor/Private/SkeletalMeshEditorCommands.h
--- a/Engine/Source/Editor/SkeletalMeshEditor/Private/SkeletalMeshEditorCommands.h
+++ b/Engine/Source/Editor/SkeletalMeshEditor/Private/SkeletalMeshEditorCommands.h
@@ -28,4 +28,7 @@ public:
 
 	// selecting mesh section using hit proxies
 	TSharedPtr<FUICommandInfo> MeshSectionSelection;
+
+	// bake materials for this skeletal mesh
+	TSharedPtr<FUICommandInfo> BakeMaterials;
 };
