Skelmesh Editor: Fix a crash when exiting the skelmesh editor but the persona scene ticks, and subsequently request capsule updates, before we remove the selection component.

#rnx

diff --git a/Engine/Source/Editor/SkeletalMeshEditor/Private/SkeletalMeshEditor.cpp b/Engine/Source/Editor/SkeletalMeshEditor/Private/SkeletalMeshEditor.cpp
--- a/Engine/Source/Editor/SkeletalMeshEditor/Private/SkeletalMeshEditor.cpp
+++ b/Engine/Source/Editor/SkeletalMeshEditor/Private/SkeletalMeshEditor.cpp
@@ -91,12 +91,6 @@ FSkeletalMeshEditor::~FSkeletalMeshEditor()
 	{
 		Editor->UnregisterForUndo(this);
 	}
-
-	if(PersonaToolkit.IsValid())
-	{
-		UPersonaSelectionComponent* SelectionComponent = PersonaToolkit->GetPreviewScene()->GetSelectionComponent();
-		SelectionComponent->GetWorld()->RemoveActor(SelectionComponent->GetOwner(), true);
-	}
 }
 
 bool IsReductionParentBaseLODUseSkeletalMeshBuildWorkflow(USkeletalMesh* SkeletalMesh, int32 TestLODIndex)
@@ -204,6 +198,21 @@ bool FSkeletalMeshEditor::OnRequestClose()
 	return bAllowClose;
 }
 
+
+void FSkeletalMeshEditor::OnClose()
+{
+	UPersonaSelectionComponent* SelectionComponent = PersonaToolkit->GetPreviewScene()->GetSelectionComponent();
+
+	if (SelectionComponent)
+	{
+		SelectionComponent->OnUpdateCapsules().Unbind();
+		SelectionComponent->GetWorld()->RemoveActor(SelectionComponent->GetOwner(), true);
+	}
+
+	ISkeletalMeshEditor::OnClose();
+}
+
+
 void FSkeletalMeshEditor::UpdateCapsules()
 {
 	UPersonaSelectionComponent* SelectionComponent = PersonaToolkit->GetPreviewScene()->GetSelectionComponent();
diff --git a/Engine/Source/Editor/SkeletalMeshEditor/Private/SkeletalMeshEditor.h b/Engine/Source/Editor/SkeletalMeshEditor/Private/SkeletalMeshEditor.h
--- a/Engine/Source/Editor/SkeletalMeshEditor/Private/SkeletalMeshEditor.h
+++ b/Engine/Source/Editor/SkeletalMeshEditor/Private/SkeletalMeshEditor.h
@@ -79,6 +79,7 @@ public:
 	
 	//~ Begin FAssetEditorToolkit Interface.
 	virtual bool OnRequestClose() override;
+	virtual void OnClose() override;
 	//~ End FAssetEditorToolkit Interface.
 
 	/** FEditorUndoClient interface */
