Fix crash when closing the persona editor and the cloth paint mode is active
#rnx

diff --git a/Engine/Source/Editor/SkeletalMeshEditor/Private/SkeletalMeshEditor.cpp b/Engine/Source/Editor/SkeletalMeshEditor/Private/SkeletalMeshEditor.cpp
--- a/Engine/Source/Editor/SkeletalMeshEditor/Private/SkeletalMeshEditor.cpp
+++ b/Engine/Source/Editor/SkeletalMeshEditor/Private/SkeletalMeshEditor.cpp
@@ -92,6 +92,14 @@ FSkeletalMeshEditor::~FSkeletalMeshEditor()
 		Editor->UnregisterForUndo(this);
 	}
 
+	//We have to do this differently, make sure we do not have any cloth paint mode active before deleting the PersonaToolkit
+	{
+		//At this point the ClothPaintMode, if exist is deactivate but not destroy
+		//This is why we do not verify if the mode is active. Calling DestroyMode on a unexisting mode is ok
+		const FEditorModeID ClothModeID = FName("ClothPaintMode");
+		GetEditorModeManager().DestroyMode(ClothModeID);
+	}
+	
 	// Reset the preview scene mesh before closing the toolkit or destroying the preview scene so the viewports can clean up properly.
 	// This is due to the viewports directly needing to use delegates on a nested USkeletalMesh parented to a debug component the editor uses.
 	// The USkeletalMesh persists beyond the lifetime of the debug component used by the toolkit or viewports,
