Fix missing cloth menu in skeletal mesh editor

Also fix static analysis warning

[Edigrate CL19942315]

diff --git a/Engine/Source/Editor/SkeletalMeshEditor/Private/SkeletalMeshEditor.cpp b/Engine/Source/Editor/SkeletalMeshEditor/Private/SkeletalMeshEditor.cpp
--- a/Engine/Source/Editor/SkeletalMeshEditor/Private/SkeletalMeshEditor.cpp
+++ b/Engine/Source/Editor/SkeletalMeshEditor/Private/SkeletalMeshEditor.cpp
@@ -1187,56 +1187,52 @@ void FSkeletalMeshEditor::ExtendMenu()
 					const int32 LodIndex = MenuContext->LodIndex;
 					const int32 SectionIndex = MenuContext->SectionIndex;
 
-					IAssetTools& AssetTools = FModuleManager::LoadModuleChecked<FAssetToolsModule>("AssetTools").Get();
-					if (AssetTools.GetAssetClassPermissionList(EAssetClassAction::CreateAsset)->PassesFilter(UClothingAssetBase::StaticClass()->GetFName()))
-					{
-						InSection.AddSubMenu(
-							"ApplyClothing",
-							LOCTEXT("MeshClickMenu_AssetApplyMenu", "Apply Clothing Data..."),
-							LOCTEXT("MeshClickMenu_AssetApplyMenu_ToolTip", "Select clothing data to apply to the selected section."),
-							FNewToolMenuChoice(FNewMenuDelegate::CreateSP(Editor.Get(), &FSkeletalMeshEditor::FillApplyClothingAssetMenu, LodIndex, SectionIndex)),
-							FToolUIActionChoice(FUIAction(
-								FExecuteAction(),
-								FCanExecuteAction::CreateSP(Editor.Get(), &FSkeletalMeshEditor::CanApplyClothing, LodIndex, SectionIndex)
-							)),
-							EUserInterfaceActionType::Button
-						);
+					InSection.AddSubMenu(
+						"ApplyClothing",
+						LOCTEXT("MeshClickMenu_AssetApplyMenu", "Apply Clothing Data..."),
+						LOCTEXT("MeshClickMenu_AssetApplyMenu_ToolTip", "Select clothing data to apply to the selected section."),
+						FNewToolMenuChoice(FNewMenuDelegate::CreateSP(Editor.Get(), &FSkeletalMeshEditor::FillApplyClothingAssetMenu, LodIndex, SectionIndex)),
+						FToolUIActionChoice(FUIAction(
+							FExecuteAction(),
+							FCanExecuteAction::CreateSP(Editor.Get(), &FSkeletalMeshEditor::CanApplyClothing, LodIndex, SectionIndex)
+						)),
+						EUserInterfaceActionType::Button
+					);
 
-						InSection.AddMenuEntry(
-							"RemoveClothing",
-							LOCTEXT("MeshClickMenu_RemoveClothing", "Remove Clothing Data"),
-							LOCTEXT("MeshClickMenu_RemoveClothing_ToolTip", "Remove the currently assigned clothing data."),
-							FSlateIcon(),
-							FToolUIActionChoice(FUIAction(
-								FExecuteAction::CreateSP(Editor.Get(), &FSkeletalMeshEditor::OnRemoveClothingAssetMenuItemClicked, LodIndex, SectionIndex),
-								FCanExecuteAction::CreateSP(Editor.Get(), &FSkeletalMeshEditor::CanRemoveClothing, LodIndex, SectionIndex)
-							))
-						);
+					InSection.AddMenuEntry(
+						"RemoveClothing",
+						LOCTEXT("MeshClickMenu_RemoveClothing", "Remove Clothing Data"),
+						LOCTEXT("MeshClickMenu_RemoveClothing_ToolTip", "Remove the currently assigned clothing data."),
+						FSlateIcon(),
+						FToolUIActionChoice(FUIAction(
+							FExecuteAction::CreateSP(Editor.Get(), &FSkeletalMeshEditor::OnRemoveClothingAssetMenuItemClicked, LodIndex, SectionIndex),
+							FCanExecuteAction::CreateSP(Editor.Get(), &FSkeletalMeshEditor::CanRemoveClothing, LodIndex, SectionIndex)
+						))
+					);
 
-						InSection.AddSubMenu(
-							"CreateClothing",
-							LOCTEXT("MeshClickMenu_CreateClothing_Label", "Create Clothing Data from Section"),
-							LOCTEXT("MeshClickMenu_CreateClothing_ToolTip", "Create a new clothing data using the selected section as a simulation mesh"),
-							FNewToolMenuChoice(FNewMenuDelegate::CreateSP(Editor.Get(), &FSkeletalMeshEditor::FillCreateClothingMenu, LodIndex, SectionIndex)),
-							FToolUIActionChoice(FUIAction(
-								FExecuteAction(),
-								FCanExecuteAction::CreateSP(Editor.Get(), &FSkeletalMeshEditor::CanCreateClothing, LodIndex, SectionIndex)
-							)),
-							EUserInterfaceActionType::Button
-						);
+					InSection.AddSubMenu(
+						"CreateClothing",
+						LOCTEXT("MeshClickMenu_CreateClothing_Label", "Create Clothing Data from Section"),
+						LOCTEXT("MeshClickMenu_CreateClothing_ToolTip", "Create a new clothing data using the selected section as a simulation mesh"),
+						FNewToolMenuChoice(FNewMenuDelegate::CreateSP(Editor.Get(), &FSkeletalMeshEditor::FillCreateClothingMenu, LodIndex, SectionIndex)),
+						FToolUIActionChoice(FUIAction(
+							FExecuteAction(),
+							FCanExecuteAction::CreateSP(Editor.Get(), &FSkeletalMeshEditor::CanCreateClothing, LodIndex, SectionIndex)
+						)),
+						EUserInterfaceActionType::Button
+					);
 
-						InSection.AddSubMenu(
-							"CreateClothingLOD",
-							LOCTEXT("MeshClickMenu_CreateClothingNewLod_Label", "Create Clothing LOD from Section"),
-							LOCTEXT("MeshClickMenu_CreateClothingNewLod_ToolTip", "Create a clothing simulation mesh from the selected section and add it as a LOD to existing clothing data."),
-							FNewToolMenuChoice(FNewMenuDelegate::CreateSP(Editor.Get(), &FSkeletalMeshEditor::FillCreateClothingLodMenu, LodIndex, SectionIndex)),
-							FToolUIActionChoice(FUIAction(
-								FExecuteAction(),
-								FCanExecuteAction::CreateSP(Editor.Get(), &FSkeletalMeshEditor::CanCreateClothingLod, LodIndex, SectionIndex)
-							)),
-							EUserInterfaceActionType::Button
-						);
-					}
+					InSection.AddSubMenu(
+						"CreateClothingLOD",
+						LOCTEXT("MeshClickMenu_CreateClothingNewLod_Label", "Create Clothing LOD from Section"),
+						LOCTEXT("MeshClickMenu_CreateClothingNewLod_ToolTip", "Create a clothing simulation mesh from the selected section and add it as a LOD to existing clothing data."),
+						FNewToolMenuChoice(FNewMenuDelegate::CreateSP(Editor.Get(), &FSkeletalMeshEditor::FillCreateClothingLodMenu, LodIndex, SectionIndex)),
+						FToolUIActionChoice(FUIAction(
+							FExecuteAction(),
+							FCanExecuteAction::CreateSP(Editor.Get(), &FSkeletalMeshEditor::CanCreateClothingLod, LodIndex, SectionIndex)
+						)),
+						EUserInterfaceActionType::Button
+					);
 				}
 			}
 		}));
@@ -1477,21 +1473,24 @@ void FSkeletalMeshEditor::HandleMeshClick(HActor* HitProxy, const FViewportClick
 	{
 		Component->SetSelectedEditorSection(HitProxy->SectionIndex);
 		Component->PushSelectionToProxy();
-	}
 
-	if(Click.GetKey() == EKeys::RightMouseButton)
-	{
-		USkeletalMeshEditorContextMenuContext* MenuContext = NewObject<USkeletalMeshEditorContextMenuContext>();
-		MenuContext->LodIndex = Component->GetPredictedLODLevel();
-		MenuContext->SectionIndex = HitProxy->SectionIndex;
-
-		FSlateApplication::Get().PushMenu(
-			FSlateApplication::Get().GetActiveTopLevelWindow().ToSharedRef(),
-			FWidgetPath(),
-			UToolMenus::Get()->GenerateWidget("SkeletalMeshEditor.MeshContextMenu", FToolMenuContext(MenuContext)),
-			FSlateApplication::Get().GetCursorPos(),
-			FPopupTransitionEffect(FPopupTransitionEffect::ContextMenu)
-			);
+		if(Click.GetKey() == EKeys::RightMouseButton)
+		{
+			USkeletalMeshEditorContextMenuContext* MenuContext = NewObject<USkeletalMeshEditorContextMenuContext>();
+			MenuContext->LodIndex = Component->GetPredictedLODLevel();
+			MenuContext->SectionIndex = HitProxy->SectionIndex;
+
+			FToolMenuContext ToolMenuContext(MenuContext);
+			InitToolMenuContext(ToolMenuContext);
+
+			FSlateApplication::Get().PushMenu(
+				FSlateApplication::Get().GetActiveTopLevelWindow().ToSharedRef(),
+				FWidgetPath(),
+				UToolMenus::Get()->GenerateWidget("SkeletalMeshEditor.MeshContextMenu", ToolMenuContext),
+				FSlateApplication::Get().GetCursorPos(),
+				FPopupTransitionEffect(FPopupTransitionEffect::ContextMenu)
+				);
+		}
 	}
 }
 
