Add some editor scripting function to allow changing the skeletalmesh sections settings.
- CastShadow
- RecomputeTangents
- VisibleInRayTracing

diff --git a/Engine/Source/Editor/SkeletalMeshEditor/Private/SkeletalMeshEditorSubsystem.cpp b/Engine/Source/Editor/SkeletalMeshEditor/Private/SkeletalMeshEditorSubsystem.cpp
--- a/Engine/Source/Editor/SkeletalMeshEditor/Private/SkeletalMeshEditorSubsystem.cpp
+++ b/Engine/Source/Editor/SkeletalMeshEditor/Private/SkeletalMeshEditorSubsystem.cpp
@@ -14,8 +14,10 @@
 #include "ObjectTools.h"
 #include "PhysicsAssetUtils.h"
 #include "PhysicsEngine/PhysicsAsset.h"
-#include "Rendering/SkeletalMeshRenderData.h"
+#include "Rendering/SkeletalMeshLODModel.h"
 #include "Rendering/SkeletalMeshLODRenderData.h"
+#include "Rendering/SkeletalMeshModel.h"
+#include "Rendering/SkeletalMeshRenderData.h"
 #include "ScopedTransaction.h"
 #include "EditorScriptingHelpers.h"
 #include "SkeletalMeshTypes.h"
@@ -105,6 +107,358 @@ int32 USkeletalMeshEditorSubsystem::GetNumSections( USkeletalMesh* SkeletalMesh,
 	return INDEX_NONE;
 }
 
+bool USkeletalMeshEditorSubsystem::GetSectionRecomputeTangent(const USkeletalMesh* SkeletalMesh, const int32 LODIndex, const int32 SectionIndex, bool& bOutRecomputeTangent)
+{
+	TGuardValue<bool> UnattendedScriptGuard(GIsRunningUnattendedScript, true);
+
+	if (!EditorScriptingHelpers::CheckIfInEditorAndPIE())
+	{
+		return false;
+	}
+
+	if (SkeletalMesh == nullptr)
+	{
+		UE_LOG(LogSkeletalMeshEditorSubsystem, Error, TEXT("GetSectionRecomputeTangent: The SkeletalMesh is null."));
+		return false;
+	}
+
+	if (!SkeletalMesh->GetResourceForRendering() || !(SkeletalMesh->GetResourceForRendering()->LODRenderData.Num() > 0))
+	{
+		UE_LOG(LogSkeletalMeshEditorSubsystem, Error, TEXT("GetSectionRecomputeTangent: The render data is null."));
+		return false;
+	}
+
+	TIndirectArray<FSkeletalMeshLODRenderData>& LodRenderData = SkeletalMesh->GetResourceForRendering()->LODRenderData;
+	if (!LodRenderData.IsValidIndex(LODIndex))
+	{
+		UE_LOG(LogSkeletalMeshEditorSubsystem, Error, TEXT("GetSectionRecomputeTangent: The LOD index is invalid."));
+		return false;
+	}
+	if (!LodRenderData[LODIndex].RenderSections.IsValidIndex(SectionIndex))
+	{
+		UE_LOG(LogSkeletalMeshEditorSubsystem, Error, TEXT("GetSectionRecomputeTangent: The section index is invalid."));
+		return false;
+	}
+
+	bOutRecomputeTangent = LodRenderData[LODIndex].RenderSections[SectionIndex].bRecomputeTangent;
+	return true;
+}
+
+bool USkeletalMeshEditorSubsystem::SetSectionRecomputeTangent(USkeletalMesh* SkeletalMesh, const int32 LODIndex, const int32 SectionIndex, const bool bRecomputeTangent)
+{
+	TGuardValue<bool> UnattendedScriptGuard(GIsRunningUnattendedScript, true);
+
+	if (!EditorScriptingHelpers::CheckIfInEditorAndPIE())
+	{
+		return false;
+	}
+
+	if (SkeletalMesh == nullptr)
+	{
+		UE_LOG(LogSkeletalMeshEditorSubsystem, Error, TEXT("SetSectionRecomputeTangent: The SkeletalMesh is null."));
+		return false;
+	}
+
+	if (!SkeletalMesh->GetImportedModel())
+	{
+		UE_LOG(LogSkeletalMeshEditorSubsystem, Error, TEXT("SetSectionRecomputeTangent: The imported model is null."));
+		return false;
+	}
+
+	if (!SkeletalMesh->GetImportedModel()->LODModels.IsValidIndex(LODIndex))
+	{
+		UE_LOG(LogSkeletalMeshEditorSubsystem, Error, TEXT("SetSectionRecomputeTangent: The LOD index is invalid."));
+		return false;
+	}
+	FSkeletalMeshLODModel& LODModel = SkeletalMesh->GetImportedModel()->LODModels[LODIndex];
+	if (!LODModel.Sections.IsValidIndex(SectionIndex))
+	{
+		UE_LOG(LogSkeletalMeshEditorSubsystem, Error, TEXT("SetSectionRecomputeTangent: The section index is invalid."));
+		return false;
+	}
+
+
+	//Make the change as a transaction and make sure the render data is rebuilt properly
+	{
+		FScopedTransaction Transaction(LOCTEXT("SetSectionRecomputeTangentTransactionName", "Set section recompute tangent"));
+		SkeletalMesh->Modify();
+		const int32 OriginalSectionIndex = LODModel.Sections[SectionIndex].OriginalDataSectionIndex;
+		FSkelMeshSourceSectionUserData& SectionUserData = LODModel.UserSectionsData.FindOrAdd(OriginalSectionIndex);
+		SectionUserData.bRecomputeTangent = bRecomputeTangent;
+	}
+	//Rebuild the render data
+	SkeletalMesh->PostEditChange();
+	//This will make the sk build synchronous
+	bool bNewValue = false;
+	GetSectionRecomputeTangent(SkeletalMesh, LODIndex, SectionIndex, bNewValue);
+	ensure(bNewValue == bRecomputeTangent);
+	return true;
+}
+
+bool USkeletalMeshEditorSubsystem::GetSectionRecomputeTangentsVertexMaskChannel(const USkeletalMesh* SkeletalMesh, const int32 LODIndex, const int32 SectionIndex, uint8& OutRecomputeTangentsVertexMaskChannel)
+{
+	TGuardValue<bool> UnattendedScriptGuard(GIsRunningUnattendedScript, true);
+
+	if (!EditorScriptingHelpers::CheckIfInEditorAndPIE())
+	{
+		return false;
+	}
+
+	if (SkeletalMesh == nullptr)
+	{
+		UE_LOG(LogSkeletalMeshEditorSubsystem, Error, TEXT("GetSectionRecomputeTangentsVertexMaskChannel: The SkeletalMesh is null."));
+		return false;
+	}
+
+	if (!SkeletalMesh->GetResourceForRendering() || !(SkeletalMesh->GetResourceForRendering()->LODRenderData.Num() > 0))
+	{
+		UE_LOG(LogSkeletalMeshEditorSubsystem, Error, TEXT("GetSectionRecomputeTangentsVertexMaskChannel: The render data is null."));
+		return false;
+	}
+
+	TIndirectArray<FSkeletalMeshLODRenderData>& LodRenderData = SkeletalMesh->GetResourceForRendering()->LODRenderData;
+	if (!LodRenderData.IsValidIndex(LODIndex))
+	{
+		UE_LOG(LogSkeletalMeshEditorSubsystem, Error, TEXT("GetSectionRecomputeTangentsVertexMaskChannel: The LOD index is invalid."));
+		return false;
+	}
+	if (!LodRenderData[LODIndex].RenderSections.IsValidIndex(SectionIndex))
+	{
+		UE_LOG(LogSkeletalMeshEditorSubsystem, Error, TEXT("GetSectionRecomputeTangentsVertexMaskChannel: The section index is invalid."));
+		return false;
+	}
+
+	OutRecomputeTangentsVertexMaskChannel = static_cast<uint8>(LodRenderData[LODIndex].RenderSections[SectionIndex].RecomputeTangentsVertexMaskChannel);
+	return true;
+}
+
+bool USkeletalMeshEditorSubsystem::SetSectionRecomputeTangentsVertexMaskChannel(USkeletalMesh* SkeletalMesh, const int32 LODIndex, const int32 SectionIndex, const uint8 RecomputeTangentsVertexMaskChannel)
+{
+	TGuardValue<bool> UnattendedScriptGuard(GIsRunningUnattendedScript, true);
+
+	if (!EditorScriptingHelpers::CheckIfInEditorAndPIE())
+	{
+		return false;
+	}
+
+	if (SkeletalMesh == nullptr)
+	{
+		UE_LOG(LogSkeletalMeshEditorSubsystem, Error, TEXT("SetSectionRecomputeTangentsVertexMaskChannel: The SkeletalMesh is null."));
+		return false;
+	}
+
+	if (!SkeletalMesh->GetImportedModel())
+	{
+		UE_LOG(LogSkeletalMeshEditorSubsystem, Error, TEXT("SetSectionRecomputeTangentsVertexMaskChannel: The imported model is null."));
+		return false;
+	}
+
+	if (!SkeletalMesh->GetImportedModel()->LODModels.IsValidIndex(LODIndex))
+	{
+		UE_LOG(LogSkeletalMeshEditorSubsystem, Error, TEXT("SetSectionRecomputeTangentsVertexMaskChannel: The LOD index is invalid."));
+		return false;
+	}
+	FSkeletalMeshLODModel& LODModel = SkeletalMesh->GetImportedModel()->LODModels[LODIndex];
+	if (!LODModel.Sections.IsValidIndex(SectionIndex))
+	{
+		UE_LOG(LogSkeletalMeshEditorSubsystem, Error, TEXT("SetSectionRecomputeTangentsVertexMaskChannel: The section index is invalid."));
+		return false;
+	}
+
+	//Make the change as a transaction and make sure the render data is rebuilt properly
+	{
+		FScopedTransaction Transaction(LOCTEXT("SetSectionRecomputeTangentsVertexMaskChannelTransactionName", "Set section recompute tangents vertex mask channel"));
+		SkeletalMesh->Modify();
+		const int32 OriginalSectionIndex = LODModel.Sections[SectionIndex].OriginalDataSectionIndex;
+		FSkelMeshSourceSectionUserData& SectionUserData = LODModel.UserSectionsData.FindOrAdd(OriginalSectionIndex);
+		SectionUserData.RecomputeTangentsVertexMaskChannel = static_cast<ESkinVertexColorChannel>(RecomputeTangentsVertexMaskChannel);
+	}
+	
+	SkeletalMesh->PostEditChange();
+	//This will make the sk build synchronous
+	uint8 NewValue;
+	GetSectionRecomputeTangentsVertexMaskChannel(SkeletalMesh, LODIndex, SectionIndex, NewValue);
+	ensure(NewValue == RecomputeTangentsVertexMaskChannel);
+	
+	return true;
+}
+
+bool USkeletalMeshEditorSubsystem::GetSectionCastShadow(const USkeletalMesh* SkeletalMesh, const int32 LODIndex, const int32 SectionIndex, bool& bOutCastShadow)
+{
+	TGuardValue<bool> UnattendedScriptGuard(GIsRunningUnattendedScript, true);
+
+	if (!EditorScriptingHelpers::CheckIfInEditorAndPIE())
+	{
+		return false;
+	}
+
+	if (SkeletalMesh == nullptr)
+	{
+		UE_LOG(LogSkeletalMeshEditorSubsystem, Error, TEXT("GetSectionCastShadow: The SkeletalMesh is null."));
+		return false;
+	}
+
+	if (!SkeletalMesh->GetResourceForRendering() || !(SkeletalMesh->GetResourceForRendering()->LODRenderData.Num() > 0))
+	{
+		UE_LOG(LogSkeletalMeshEditorSubsystem, Error, TEXT("GetSectionCastShadow: The render data is null."));
+		return false;
+	}
+
+	TIndirectArray<FSkeletalMeshLODRenderData>& LodRenderData = SkeletalMesh->GetResourceForRendering()->LODRenderData;
+	if (!LodRenderData.IsValidIndex(LODIndex))
+	{
+		UE_LOG(LogSkeletalMeshEditorSubsystem, Error, TEXT("GetSectionCastShadow: The LOD index is invalid."));
+		return false;
+	}
+	if (!LodRenderData[LODIndex].RenderSections.IsValidIndex(SectionIndex))
+	{
+		UE_LOG(LogSkeletalMeshEditorSubsystem, Error, TEXT("GetSectionCastShadow: The section index is invalid."));
+		return false;
+	}
+
+	bOutCastShadow = LodRenderData[LODIndex].RenderSections[SectionIndex].bCastShadow;
+	return true;
+}
+
+bool USkeletalMeshEditorSubsystem::SetSectionCastShadow(USkeletalMesh* SkeletalMesh, const int32 LODIndex, const int32 SectionIndex, const bool bCastShadow)
+{
+	TGuardValue<bool> UnattendedScriptGuard(GIsRunningUnattendedScript, true);
+
+	if (!EditorScriptingHelpers::CheckIfInEditorAndPIE())
+	{
+		return false;
+	}
+
+	if (SkeletalMesh == nullptr)
+	{
+		UE_LOG(LogSkeletalMeshEditorSubsystem, Error, TEXT("SetSectionCastShadow: The SkeletalMesh is null."));
+		return false;
+	}
+
+	if (!SkeletalMesh->GetImportedModel())
+	{
+		UE_LOG(LogSkeletalMeshEditorSubsystem, Error, TEXT("SetSectionCastShadow: The imported model is null."));
+		return false;
+	}
+
+	if (!SkeletalMesh->GetImportedModel()->LODModels.IsValidIndex(LODIndex))
+	{
+		UE_LOG(LogSkeletalMeshEditorSubsystem, Error, TEXT("SetSectionCastShadow: The LOD index is invalid."));
+		return false;
+	}
+	FSkeletalMeshLODModel& LODModel = SkeletalMesh->GetImportedModel()->LODModels[LODIndex];
+	if (!LODModel.Sections.IsValidIndex(SectionIndex))
+	{
+		UE_LOG(LogSkeletalMeshEditorSubsystem, Error, TEXT("SetSectionCastShadow: The section index is invalid."));
+		return false;
+	}
+
+	//Make the change as a transaction and make sure the render data is rebuilt properly
+	{
+		FScopedTransaction Transaction(LOCTEXT("SetSectionCastShadowTransactionName", "Set section cast shadow"));
+		SkeletalMesh->Modify();
+		const int32 OriginalSectionIndex = LODModel.Sections[SectionIndex].OriginalDataSectionIndex;
+		FSkelMeshSourceSectionUserData& SectionUserData = LODModel.UserSectionsData.FindOrAdd(OriginalSectionIndex);
+		SectionUserData.bCastShadow = bCastShadow;
+	}
+	//Rebuild the render data
+	SkeletalMesh->PostEditChange();
+	//This will make the sk build synchronous
+	bool bNewValue = false;
+	GetSectionCastShadow(SkeletalMesh, LODIndex, SectionIndex, bNewValue);
+	ensure(bNewValue == bCastShadow);
+
+	return true;
+}
+
+bool USkeletalMeshEditorSubsystem::GetSectionVisibleInRayTracing(const USkeletalMesh* SkeletalMesh, const int32 LODIndex, const int32 SectionIndex, bool& bOutVisibleInRayTracing)
+{
+	TGuardValue<bool> UnattendedScriptGuard(GIsRunningUnattendedScript, true);
+
+	if (!EditorScriptingHelpers::CheckIfInEditorAndPIE())
+	{
+		return false;
+	}
+
+	if (SkeletalMesh == nullptr)
+	{
+		UE_LOG(LogSkeletalMeshEditorSubsystem, Error, TEXT("GetSectionVisibleInRayTracing: The SkeletalMesh is null."));
+		return false;
+	}
+
+	if (!SkeletalMesh->GetResourceForRendering() || !(SkeletalMesh->GetResourceForRendering()->LODRenderData.Num() > 0))
+	{
+		UE_LOG(LogSkeletalMeshEditorSubsystem, Error, TEXT("GetSectionVisibleInRayTracing: The render data is null."));
+		return false;
+	}
+
+	TIndirectArray<FSkeletalMeshLODRenderData>& LodRenderData = SkeletalMesh->GetResourceForRendering()->LODRenderData;
+	if (!LodRenderData.IsValidIndex(LODIndex))
+	{
+		UE_LOG(LogSkeletalMeshEditorSubsystem, Error, TEXT("GetSectionVisibleInRayTracing: The LOD index is invalid."));
+		return false;
+	}
+	if (!LodRenderData[LODIndex].RenderSections.IsValidIndex(SectionIndex))
+	{
+		UE_LOG(LogSkeletalMeshEditorSubsystem, Error, TEXT("GetSectionVisibleInRayTracing: The section index is invalid."));
+		return false;
+	}
+
+	bOutVisibleInRayTracing = LodRenderData[LODIndex].RenderSections[SectionIndex].bVisibleInRayTracing;
+	return true;
+}
+
+bool USkeletalMeshEditorSubsystem::SetSectionVisibleInRayTracing(USkeletalMesh* SkeletalMesh, const int32 LODIndex, const int32 SectionIndex, const bool bVisibleInRayTracing)
+{
+	TGuardValue<bool> UnattendedScriptGuard(GIsRunningUnattendedScript, true);
+
+	if (!EditorScriptingHelpers::CheckIfInEditorAndPIE())
+	{
+		return false;
+	}
+
+	if (SkeletalMesh == nullptr)
+	{
+		UE_LOG(LogSkeletalMeshEditorSubsystem, Error, TEXT("SetSectionVisibleInRayTracing: The SkeletalMesh is null."));
+		return false;
+	}
+
+	if (!SkeletalMesh->GetImportedModel())
+	{
+		UE_LOG(LogSkeletalMeshEditorSubsystem, Error, TEXT("SetSectionVisibleInRayTracing: The imported model is null."));
+		return false;
+	}
+
+	if (!SkeletalMesh->GetImportedModel()->LODModels.IsValidIndex(LODIndex))
+	{
+		UE_LOG(LogSkeletalMeshEditorSubsystem, Error, TEXT("SetSectionVisibleInRayTracing: The LOD index is invalid."));
+		return false;
+	}
+	FSkeletalMeshLODModel& LODModel = SkeletalMesh->GetImportedModel()->LODModels[LODIndex];
+	if (!LODModel.Sections.IsValidIndex(SectionIndex))
+	{
+		UE_LOG(LogSkeletalMeshEditorSubsystem, Error, TEXT("SetSectionVisibleInRayTracing: The section index is invalid."));
+		return false;
+	}
+
+	//Make the change as a transaction and make sure the render data is rebuilt properly
+	{
+		FScopedTransaction Transaction(LOCTEXT("SetSectionVisibleInRayTracingTransactionName", "Set section visible in ray tracing"));
+		SkeletalMesh->Modify();
+		const int32 OriginalSectionIndex = LODModel.Sections[SectionIndex].OriginalDataSectionIndex;
+		FSkelMeshSourceSectionUserData& SectionUserData = LODModel.UserSectionsData.FindOrAdd(OriginalSectionIndex);
+		SectionUserData.bVisibleInRayTracing = bVisibleInRayTracing;
+	}
+
+	SkeletalMesh->PostEditChange();
+	//This will make the sk build synchronous
+	bool bNewValue = false;
+	GetSectionVisibleInRayTracing(SkeletalMesh, LODIndex, SectionIndex, bNewValue);
+	ensure(bNewValue == bVisibleInRayTracing);
+
+	return true;
+}
+
 int32 USkeletalMeshEditorSubsystem::GetLODMaterialSlot( USkeletalMesh* SkeletalMesh, int32 LODIndex, int32 SectionIndex )
 {
 	TGuardValue<bool> UnattendedScriptGuard( GIsRunningUnattendedScript, true );
diff --git a/Engine/Source/Editor/SkeletalMeshEditor/Public/SkeletalMeshEditorSubsystem.h b/Engine/Source/Editor/SkeletalMeshEditor/Public/SkeletalMeshEditorSubsystem.h
--- a/Engine/Source/Editor/SkeletalMeshEditor/Public/SkeletalMeshEditorSubsystem.h
+++ b/Engine/Source/Editor/SkeletalMeshEditor/Public/SkeletalMeshEditorSubsystem.h
@@ -54,6 +54,94 @@ public:
 	UFUNCTION(BlueprintPure, Category = "Skeletal Mesh Utilities")
 	int32 GetNumSections(USkeletalMesh* SkeletalMesh, int32 LODIndex);
 
+	/** Get bRecomputeTangent from a section of a LOD of a Skeletal Mesh
+	 *
+	 * @param SkeletalMesh			Mesh to get number of vertices from.
+	 * @param LodIndex				Index of the mesh LOD.
+	 * @param SectionIndex			Index of the LOD section.
+	 * @param bOutRecomputeTangent	The function will set the bRecomputeTangent used by the section
+	 * @return false if invalid mesh or LOD index or section index.
+	 */
+	UFUNCTION(BlueprintPure, Category = "Skeletal Mesh Utilities")
+	bool GetSectionRecomputeTangent(const USkeletalMesh* SkeletalMesh, const int32 LODIndex, const int32 SectionIndex, bool& bOutRecomputeTangent);
+
+	/** Set bRecomputeTangent for a section of a LOD of a Skeletal Mesh.
+	 *
+	 * @param SkeletalMesh			Mesh to get number of vertices from.
+	 * @param LodIndex				Index of the mesh LOD.
+	 * @param SectionIndex			Index of the LOD section.
+	 * @param bRecomputeTangent	The function will set the bRecomputeTangent used by the section
+	 * @return false if invalid mesh or LOD index or section index.
+	 */
+	UFUNCTION(BlueprintCallable, Category = "Skeletal Mesh Utilities")
+	bool SetSectionRecomputeTangent(USkeletalMesh* SkeletalMesh, const int32 LODIndex, const int32 SectionIndex, const bool bRecomputeTangent);
+
+	/** Get RecomputeTangentsVertexMaskChannel from a section of a LOD of a Skeletal Mesh
+	 *
+	 * @param SkeletalMesh			Mesh to get number of vertices from.
+	 * @param LodIndex				Index of the mesh LOD.
+	 * @param SectionIndex			Index of the LOD section.
+	 * @param OutRecomputeTangentsVertexMaskChannel	The function will set the RecomputeTangentsVertexMaskChannel used by the section
+	 * @return false if invalid mesh or LOD index or section index.
+	 */
+	UFUNCTION(BlueprintPure, Category = "Skeletal Mesh Utilities")
+	bool GetSectionRecomputeTangentsVertexMaskChannel(const USkeletalMesh* SkeletalMesh, const int32 LODIndex, const int32 SectionIndex, uint8& OutRecomputeTangentsVertexMaskChannel);
+
+	/** Set RecomputeTangentsVertexMaskChannel for a section of a LOD of a Skeletal Mesh.
+	 *
+	 * @param SkeletalMesh			Mesh to get number of vertices from.
+	 * @param LodIndex				Index of the mesh LOD.
+	 * @param SectionIndex			Index of the LOD section.
+	 * @param RecomputeTangentsVertexMaskChannel	The function will set the RecomputeTangentsVertexMaskChannel used by the section
+	 * @return false if invalid mesh or LOD index or section index.
+	 */
+	UFUNCTION(BlueprintCallable, Category = "Skeletal Mesh Utilities")
+	bool SetSectionRecomputeTangentsVertexMaskChannel(USkeletalMesh* SkeletalMesh, const int32 LODIndex, const int32 SectionIndex, const uint8 RecomputeTangentsVertexMaskChannel);
+
+	/** Get bCastShadow from a section of a LOD of a Skeletal Mesh
+	 *
+	 * @param SkeletalMesh			Mesh to get number of vertices from.
+	 * @param LodIndex				Index of the mesh LOD.
+	 * @param SectionIndex			Index of the LOD section.
+	 * @param bOutCastShadow	The function will set the bCastShadow used by the section
+	 * @return false if invalid mesh or LOD index or section index.
+	 */
+	UFUNCTION(BlueprintPure, Category = "Skeletal Mesh Utilities")
+	bool GetSectionCastShadow(const USkeletalMesh* SkeletalMesh, const int32 LODIndex, const int32 SectionIndex, bool& bOutCastShadow);
+
+	/** Set bCastShadow for a section of a LOD of a Skeletal Mesh.
+	 *
+	 * @param SkeletalMesh			Mesh to get number of vertices from.
+	 * @param LodIndex				Index of the mesh LOD.
+	 * @param SectionIndex			Index of the LOD section.
+	 * @param bCastShadow	The function will set the bCastShadow used by the section
+	 * @return false if invalid mesh or LOD index or section index.
+	 */
+	UFUNCTION(BlueprintCallable, Category = "Skeletal Mesh Utilities")
+	bool SetSectionCastShadow(USkeletalMesh* SkeletalMesh, const int32 LODIndex, const int32 SectionIndex, const bool bCastShadow);
+
+	/** Get bVisibleInRayTracing from a section of a LOD of a Skeletal Mesh
+	 *
+	 * @param SkeletalMesh			Mesh to get number of vertices from.
+	 * @param LodIndex				Index of the mesh LOD.
+	 * @param SectionIndex			Index of the LOD section.
+	 * @param bOutVisibleInRayTracing	The function will set the bVisibleInRayTracing used by the section
+	 * @return false if invalid mesh or LOD index or section index.
+	 */
+	UFUNCTION(BlueprintPure, Category = "Skeletal Mesh Utilities")
+	bool GetSectionVisibleInRayTracing(const USkeletalMesh* SkeletalMesh, const int32 LODIndex, const int32 SectionIndex, bool& bOutVisibleInRayTracing);
+
+	/** Set bVisibleInRayTracing for a section of a LOD of a Skeletal Mesh.
+	 *
+	 * @param SkeletalMesh			Mesh to get number of vertices from.
+	 * @param LodIndex				Index of the mesh LOD.
+	 * @param SectionIndex			Index of the LOD section.
+	 * @param bVisibleInRayTracing	The function will set the bVisibleInRayTracing used by the section
+	 * @return false if invalid mesh or LOD index or section index.
+	 */
+	UFUNCTION(BlueprintCallable, Category = "Skeletal Mesh Utilities")
+	bool SetSectionVisibleInRayTracing(USkeletalMesh* SkeletalMesh, const int32 LODIndex, const int32 SectionIndex, const bool bVisibleInRayTracing);
+
 	/**
 	 * Gets the material slot used for a specific LOD section.
 	 * @param	SkeletalMesh		SkeletalMesh to get the material index from.
