Fix for Persona mesh details text overlay blocking hit tests
* Fix for Persona mesh sections showing as selectable when not.
* Fix for editor viewport cursor style not updating if using orbit camera mode.

diff --git a/Engine/Source/Editor/AnimationBlueprintEditor/Private/AnimationBlueprintEditor.cpp b/Engine/Source/Editor/AnimationBlueprintEditor/Private/AnimationBlueprintEditor.cpp
--- a/Engine/Source/Editor/AnimationBlueprintEditor/Private/AnimationBlueprintEditor.cpp
+++ b/Engine/Source/Editor/AnimationBlueprintEditor/Private/AnimationBlueprintEditor.cpp
@@ -504,6 +504,8 @@ void FAnimationBlueprintEditor::InitAnimationBlueprintEditor(const EToolkitMode:
 	// Register for notifications when settings change
 	AnimationBlueprintEditorSettingsChangedHandle = GetMutableDefault<UAnimationBlueprintEditorSettings>()->RegisterOnUpdateSettings(
 		UAnimationBlueprintEditorSettings::FOnUpdateSettingsMulticaster::FDelegate::CreateSP(this, &FAnimationBlueprintEditor::HandleUpdateSettings));
+
+	PersonaToolkit->GetPreviewScene()->SetAllowMeshHitProxies(false);
 }
 
 void FAnimationBlueprintEditor::BindCommands()
diff --git a/Engine/Source/Editor/AnimationEditor/Private/AnimationEditor.cpp b/Engine/Source/Editor/AnimationEditor/Private/AnimationEditor.cpp
--- a/Engine/Source/Editor/AnimationEditor/Private/AnimationEditor.cpp
+++ b/Engine/Source/Editor/AnimationEditor/Private/AnimationEditor.cpp
@@ -173,6 +173,8 @@ void FAnimationEditor::InitAnimationEditor(const EToolkitMode::Type Mode, const
 	RegenerateMenusAndToolbars();
 
 	OpenNewAnimationDocumentTab(AnimationAsset);
+
+	PersonaToolkit->GetPreviewScene()->SetAllowMeshHitProxies(false);
 }
 
 FName FAnimationEditor::GetToolkitFName() const
diff --git a/Engine/Source/Editor/Persona/Private/AnimationEditorPreviewScene.cpp b/Engine/Source/Editor/Persona/Private/AnimationEditorPreviewScene.cpp
--- a/Engine/Source/Editor/Persona/Private/AnimationEditorPreviewScene.cpp
+++ b/Engine/Source/Editor/Persona/Private/AnimationEditorPreviewScene.cpp
@@ -1138,6 +1138,12 @@ bool FAnimationEditorPreviewScene::AllowMeshHitProxies() const
 void FAnimationEditorPreviewScene::SetAllowMeshHitProxies(bool bState)
 {
 	bEnableMeshHitProxies = bState;
+
+	if (GetPreviewMeshComponent())
+	{
+		GetPreviewMeshComponent()->bSelectable = bEnableMeshHitProxies;
+		GetPreviewMeshComponent()->MarkRenderStateDirty();
+	}
 }
 
 void FAnimationEditorPreviewScene::FlagTickable()
diff --git a/Engine/Source/Editor/Persona/Private/SAnimationEditorViewport.cpp b/Engine/Source/Editor/Persona/Private/SAnimationEditorViewport.cpp
--- a/Engine/Source/Editor/Persona/Private/SAnimationEditorViewport.cpp
+++ b/Engine/Source/Editor/Persona/Private/SAnimationEditorViewport.cpp
@@ -133,6 +133,7 @@ TSharedRef<FEditorViewportClient> SAnimationEditorViewport::MakeEditorViewportCl
 TSharedPtr<SWidget> SAnimationEditorViewport::MakeViewportToolbar()
 {
 	return SAssignNew(ViewportToolbar, SAnimViewportToolBar, TabBodyPtr.Pin(), SharedThis(this))
+		.Visibility(EVisibility::SelfHitTestInvisible)
 		.Cursor(EMouseCursor::Default)
 		.Extenders(Extenders)
 		.ContextName(ContextName)
diff --git a/Engine/Source/Editor/SkeletalMeshEditor/Private/SkeletalMeshEditor.cpp b/Engine/Source/Editor/SkeletalMeshEditor/Private/SkeletalMeshEditor.cpp
--- a/Engine/Source/Editor/SkeletalMeshEditor/Private/SkeletalMeshEditor.cpp
+++ b/Engine/Source/Editor/SkeletalMeshEditor/Private/SkeletalMeshEditor.cpp
@@ -272,11 +272,6 @@ void FSkeletalMeshEditor::InitSkeletalMeshEditor(const EToolkitMode::Type Mode,
 	// Set up mesh click selection
 	PreviewScene->RegisterOnMeshClick(FOnMeshClick::CreateSP(this, &FSkeletalMeshEditor::HandleMeshClick));
 	PreviewScene->SetAllowMeshHitProxies(true);
-	if(PreviewScene->GetPreviewMeshComponent())
-	{
-		PreviewScene->GetPreviewMeshComponent()->bSelectable = PreviewScene->AllowMeshHitProxies();
-		PreviewScene->GetPreviewMeshComponent()->MarkRenderStateDirty();
-	}
 
 	// Make sure we get told when the editor mode changes so we can switch to the appropriate tab
 	// if there's a toolbox available.
diff --git a/Engine/Source/Editor/SkeletonEditor/Private/SkeletonEditor.cpp b/Engine/Source/Editor/SkeletonEditor/Private/SkeletonEditor.cpp
--- a/Engine/Source/Editor/SkeletonEditor/Private/SkeletonEditor.cpp
+++ b/Engine/Source/Editor/SkeletonEditor/Private/SkeletonEditor.cpp
@@ -125,6 +125,8 @@ void FSkeletonEditor::InitSkeletonEditor(const EToolkitMode::Type Mode, const TS
 	ExtendMenu();
 	ExtendToolbar();
 	RegenerateMenusAndToolbars();
+
+	PersonaToolkit->GetPreviewScene()->SetAllowMeshHitProxies(false);
 }
 
 FName FSkeletonEditor::GetToolkitFName() const
diff --git a/Engine/Source/Editor/UnrealEd/Private/EditorViewportClient.cpp b/Engine/Source/Editor/UnrealEd/Private/EditorViewportClient.cpp
--- a/Engine/Source/Editor/UnrealEd/Private/EditorViewportClient.cpp
+++ b/Engine/Source/Editor/UnrealEd/Private/EditorViewportClient.cpp
@@ -1617,7 +1617,7 @@ EMouseCursor::Type FEditorViewportClient::GetCursor(FViewport* InViewport,int32
 			HHitProxy* HitProxy = InViewport->GetHitProxy(X,Y);
 
 			// Change the mouse cursor if the user is hovering over something they can interact with.
-			if( HitProxy && !bUsingOrbitCamera )
+			if( HitProxy && !IsTracking() )
 			{
 				MouseCursor = HitProxy->GetMouseCursor();
 				bShouldCheckHitProxy = true;
