Fix truncation warnings in SkeletalMeshEditor

#rnx

diff --git a/Engine/Source/Editor/SkeletalMeshEditor/Private/SkeletalMeshEditor.cpp b/Engine/Source/Editor/SkeletalMeshEditor/Private/SkeletalMeshEditor.cpp
--- a/Engine/Source/Editor/SkeletalMeshEditor/Private/SkeletalMeshEditor.cpp
+++ b/Engine/Source/Editor/SkeletalMeshEditor/Private/SkeletalMeshEditor.cpp
@@ -1062,7 +1062,7 @@ void FSkeletalMeshEditor::ApplyClothing(UClothingAssetBase* InAsset, int32 InLod
 				//Successful bind so set the SectionUserData
 				int32 AssetIndex = INDEX_NONE;
 				check(Mesh->GetMeshClothingAssets().Find(ClothingAsset, AssetIndex));
-				OriginalSectionData.CorrespondClothAssetIndex = AssetIndex;
+				OriginalSectionData.CorrespondClothAssetIndex = static_cast<int16>(AssetIndex);
 				OriginalSectionData.ClothingData.AssetGuid = ClothingAsset->GetAssetGuid();
 				OriginalSectionData.ClothingData.AssetLodIndex = InClothingLod;
 			}
