New skin weight editing BP/Python API.

diff --git a/Engine/Source/Editor/SkeletalMeshEditor/Private/SkinWeightModifier.cpp b/Engine/Source/Editor/SkeletalMeshEditor/Private/SkinWeightModifier.cpp
new file mode 100644
--- /dev/null
+++ b/Engine/Source/Editor/SkeletalMeshEditor/Private/SkinWeightModifier.cpp
@@ -0,0 +1,418 @@
+﻿// Copyright Epic Games, Inc. All Rights Reserved.
+
+#include "SkinWeightModifier.h"
+
+#include "MeshDescription.h"
+#include "ScopedTransaction.h"
+#include "SkeletalMeshAttributes.h"
+#include "Rendering/SkeletalMeshModel.h"
+
+#define LOCTEXT_NAMESPACE "SkinWeightModifier"
+
+bool USkinWeightModifier::SetSkeletalMesh(USkeletalMesh* InMesh)
+{
+	// prepare to load mesh data
+	Reset();
+	
+	// validate supplied skeletal mesh exists
+	if (!InMesh)
+	{
+		UE_LOG(LogAnimation, Error, TEXT("Skin Weight Modifier: No skeletal mesh supplied to load."));
+		return false;
+	}
+
+	// verify mesh has an imported model
+	const FSkeletalMeshModel* SkeletalMeshModel = InMesh->GetImportedModel();
+	if (!SkeletalMeshModel)
+	{
+		UE_LOG(LogAnimation, Error, TEXT("Skin Weight Modifier: Supplied skeletal mesh does not have an imported model."));
+		return false;
+	}
+
+	// verify that we have a LOD 0 mesh
+	if (SkeletalMeshModel->LODModels.IsEmpty())
+	{
+		UE_LOG(LogAnimation, Error, TEXT("Skin Weight Modifier: Supplied skeletal mesh does not have a LOD 0 mesh."));
+		return false;
+	}
+
+	// verify user is not trying to modify one of the core engine assets
+	if (InMesh->GetPathName().StartsWith(TEXT("/Engine/")))
+	{
+		UE_LOG(LogAnimation, Error, TEXT("Skin Weight Modifier: Cannot modify built-in engine asset."));
+		return false;
+	}
+
+	// store pointer to mesh 
+	Mesh = InMesh;
+
+	// store mesh description to edit
+	MeshDescription = MakeUnique<FMeshDescription>();
+	InMesh->GetMeshDescription(LODIndex, *MeshDescription);
+	if (MeshDescription->IsEmpty())
+	{
+		UE_LOG(LogAnimation, Error, TEXT("Skin Weight Modifier: mesh description is empty."));
+		return false;
+	}
+
+	// load weights into memory for subsequent editing
+	const FSkeletalMeshConstAttributes MeshAttribs(*MeshDescription);
+	const FSkinWeightsVertexAttributesConstRef VertexSkinWeights = MeshAttribs.GetVertexSkinWeights();
+	const int32 NumVertices = MeshDescription->Vertices().Num();
+	Weights.SetNum(NumVertices);
+	const FReferenceSkeleton& RefSkeleton = Mesh->GetRefSkeleton();
+	for (int32 VertexIndex = 0; VertexIndex < NumVertices; VertexIndex++)
+	{
+		const FVertexID VertexID(VertexIndex);
+		int32 InfluenceIndex = 0;
+		for (UE::AnimationCore::FBoneWeight BoneWeight: VertexSkinWeights.Get(VertexID))
+		{
+			const int32 BoneIndex = BoneWeight.GetBoneIndex();
+			const FName BoneName = RefSkeleton.GetBoneName(BoneIndex);
+			const float Weight = BoneWeight.GetWeight();
+			Weights[VertexIndex].Add(BoneName, Weight);
+			++InfluenceIndex;
+		}
+
+		if (InfluenceIndex > MAX_TOTAL_INFLUENCES)
+		{
+			UE_LOG(LogAnimation, Warning, TEXT("Skin Weight Modifier: Supplied skeletal mesh has a vertex with too many influences."));
+		}
+	}
+
+	return true;
+}
+
+bool USkinWeightModifier::CommitWeightsToSkeletalMesh()
+{
+	if (!Mesh)
+	{
+		UE_LOG(LogAnimation, Warning, TEXT("Skin Weight Modifier: No mesh loaded. Cannot apply weights to skeletal mesh."));
+		return false;
+	}
+
+	// update weights in the mesh description
+	{
+		FSkeletalMeshAttributes MeshAttribs(*MeshDescription);
+		FSkinWeightsVertexAttributesRef VertexSkinWeights = MeshAttribs.GetVertexSkinWeights();
+	
+		UE::AnimationCore::FBoneWeightsSettings Settings;
+		Settings.SetNormalizeType(UE::AnimationCore::EBoneWeightNormalizeType::Always);
+
+		TArray<UE::AnimationCore::FBoneWeight> SourceBoneWeights;
+		SourceBoneWeights.Reserve(UE::AnimationCore::MaxInlineBoneWeightCount);
+
+		const int32 NumVertices = MeshDescription->Vertices().Num();
+		const TArray<FName> AllBoneNames = GetAllBoneNames();
+		for (int32 VertexIndex = 0; VertexIndex < NumVertices; VertexIndex++)
+		{
+			SourceBoneWeights.Reset();
+
+			const TMap<FName, float>& VertexWeights = Weights[VertexIndex];
+			if (VertexWeights.IsEmpty())
+			{
+				// in case where user has removed ALL influences from a vertex, revert back to the root bone
+				const FName RootName = Mesh->GetRefSkeleton().GetBoneName(0);
+				const int32 BoneIndex = AllBoneNames.IndexOfByKey(RootName);
+				constexpr float FullWeight = 1.0f;
+				SourceBoneWeights.Add(UE::AnimationCore::FBoneWeight(BoneIndex, FullWeight));
+			}
+			else
+			{
+				for (const TTuple<FName, float>& SingleBoneWeight : VertexWeights)
+				{
+					const FName BoneName = SingleBoneWeight.Key;
+					const int32 BoneIndex = AllBoneNames.IndexOfByKey(BoneName);
+					const float Weight = SingleBoneWeight.Value;
+					SourceBoneWeights.Add(UE::AnimationCore::FBoneWeight(BoneIndex, Weight));
+				}
+			}
+
+			VertexSkinWeights.Set(FVertexID(VertexIndex), UE::AnimationCore::FBoneWeights::Create(SourceBoneWeights, Settings));
+		}
+	}
+	
+	// push changes to the skeletal mesh asset
+	{
+		// flush any pending rendering commands, which might touch a component while we are rebuilding it's mesh
+		FlushRenderingCommands();
+		// open an undo transaction on the skeletal mesh
+		FScopedTransaction Transaction(LOCTEXT("ApplySkinWeightModifications", "Apply Skin Weights"));
+		// update skeletal mesh LOD (cf. USkeletalMesh::CommitMeshDescription)
+		Mesh->CommitMeshDescription(LODIndex, *MeshDescription);
+		// rebuilds mesh
+		Mesh->PostEditChange();
+	}
+	
+	return true;
+}
+
+TMap<FName, float> USkinWeightModifier::GetVertexWeights(int32 VertexID)
+{
+	if (!Weights.IsValidIndex(VertexID))
+	{
+		UE_LOG(LogAnimation, Warning, TEXT("Skin Weight Modifier: Supplied vertex index is not valid for this mesh."));
+		return TMap<FName, float>();
+	}
+
+	return Weights[VertexID];
+}
+
+bool USkinWeightModifier::SetVertexWeights(
+	const int32 VertexID,
+	const TMap<FName, float>& InWeights,
+	const bool bReplaceAll)
+{
+	if (!Mesh)
+	{
+		UE_LOG(LogAnimation, Warning, TEXT("Skin Weight Modifier: No mesh loaded. Cannot set weights."));
+		return false;
+	}
+
+	if (!Weights.IsValidIndex(VertexID))
+	{
+		UE_LOG(LogAnimation, Warning, TEXT("Skin Weight Modifier: Supplied vertex index is not valid on this mesh. Cannot set weights."));
+		return false;
+	}
+
+	TMap<FName, float>& VertexWeights = Weights[VertexID];
+	if (bReplaceAll)
+	{
+		// replace weights entirely
+		VertexWeights = InWeights;
+	}
+	else
+	{
+		// replace only the input weights
+		for (const TTuple<FName, float>& InWeight : InWeights)
+		{
+			VertexWeights.FindOrAdd(InWeight.Key) = InWeight.Value;
+		}
+	}
+	
+	return true;
+}
+
+bool USkinWeightModifier::NormalizeVertexWeights(const int32 VertexID)
+{
+	if (!Mesh)
+	{
+		UE_LOG(LogAnimation, Warning, TEXT("Skin Weight Modifier: No mesh loaded. Cannot normalize weights."));
+		return false;
+	}
+
+	if (!Weights.IsValidIndex(VertexID))
+	{
+		UE_LOG(LogAnimation, Warning, TEXT("Skin Weight Modifier: Supplied vertex ID was invalid. Cannot normalize vertex weights."));
+		return false;
+	}
+
+	return NormalizeVertexInternal(VertexID);
+}
+
+bool USkinWeightModifier::NormalizeAllWeights()
+{
+	if (!(Mesh && MeshDescription))
+	{
+		UE_LOG(LogAnimation, Warning, TEXT("Skin Weight Modifier: No mesh loaded. Cannot normalize weights."));
+		return false;
+	}
+	
+	const int32 NumVertices = MeshDescription->Vertices().Num();
+	for (int32 VertexID=0; VertexID<NumVertices; ++VertexID)
+	{
+		NormalizeVertexInternal(VertexID);
+	}
+
+	return true;
+}
+
+bool USkinWeightModifier::EnforceMaxInfluences(const int32 MaxInfluences)
+{
+	if (!MeshDescription)
+	{
+		UE_LOG(LogAnimation, Warning, TEXT("Skin Weight Modifier: No mesh loaded. Cannot enforce max influences."));
+		return false;
+	}
+	
+	int32 Max = MaxInfluences;
+	if (Max < 0)
+	{
+		Max = UE::AnimationCore::MaxInlineBoneWeightCount;
+	}
+
+	const int32 NumVertices = MeshDescription->Vertices().Num();
+	for (int32 VertexID=0; VertexID<NumVertices; ++VertexID)
+	{
+		EnforceMaxInfluenceInternal(VertexID, Max);
+	}
+
+	return true;
+}
+
+bool USkinWeightModifier::PruneVertexWeights(const int32 VertexID, const float WeightThreshold)
+{
+	if (!Mesh)
+	{
+		UE_LOG(LogAnimation, Warning, TEXT("Skin Weight Modifier: No mesh loaded. Cannot prune weights."));
+		return false;
+	}
+
+	if (!Weights.IsValidIndex(VertexID))
+	{
+		UE_LOG(LogAnimation, Warning, TEXT("Skin Weight Modifier: Supplied vertex ID was invalid. Cannot prune weights."));
+		return false;
+	}
+
+	PruneVertexWeightInternal(VertexID, WeightThreshold);
+	
+	return true;
+}
+
+bool USkinWeightModifier::PruneAllWeights(const float WeightThreshold)
+{
+	if (!Mesh)
+	{
+		UE_LOG(LogAnimation, Warning, TEXT("Skin Weight Modifier: No mesh loaded. Cannot prune weights."));
+		return false;
+	}
+
+	const int32 NumVertices = MeshDescription->Vertices().Num();
+	for (int32 i=0; i<NumVertices; ++i)
+	{
+		PruneVertexWeightInternal(i, WeightThreshold);
+	}
+	
+	return true;
+}
+
+int32 USkinWeightModifier::GetNumVertices()
+{
+	if (!MeshDescription)
+	{
+		UE_LOG(LogAnimation, Warning, TEXT("Skin Weight Modifier: No mesh loaded. Cannot get vertex count."));
+		return INDEX_NONE;
+	}
+
+	return MeshDescription->Vertices().Num();
+}
+
+TArray<FName> USkinWeightModifier::GetAllBoneNames()
+{
+	if (!Mesh)
+	{
+		UE_LOG(LogAnimation, Warning, TEXT("Skin Weight Modifier: No mesh loaded. Cannot get bones."));
+		return TArray<FName>();
+	}
+
+	const FReferenceSkeleton& RefSkeleton = Mesh->GetRefSkeleton();
+	const int32 NumBones = RefSkeleton.GetNum();
+	TArray<FName> AllBoneNames;
+	for (int32 i=0; i<NumBones; ++i)
+	{
+		AllBoneNames.Add(RefSkeleton.GetBoneName(i));
+	}
+	return AllBoneNames;
+}
+
+bool USkinWeightModifier::NormalizeVertexInternal(const int32 VertexID)
+{
+	// handle case where user has removed all weights on a vertex (probably a mistake so warn them)
+	TMap<FName, float>& VertexWeights = Weights[VertexID];
+	if (VertexWeights.IsEmpty())
+	{
+		const FName RootBone = Mesh->GetRefSkeleton().GetBoneName(0);
+		VertexWeights.Add(RootBone, 1.f);
+		UE_LOG(LogAnimation, Warning, TEXT("Skin Weight Modifier: Supplied vertex did not have any bone influences. Applied all weight to root bone."));
+		return true;
+	}
+
+	// calculate total weight on all vertices
+	float TotalWeight = 0.f;
+	for (const TTuple<FName, float>& Weight : VertexWeights)
+	{
+		TotalWeight += Weight.Value;
+	}
+
+	// handle case where total weight is near zero
+	if (TotalWeight <= KINDA_SMALL_NUMBER)
+	{
+		// in this case, we just split the weight evenly across all available influences
+		const float SplitWeight = 1.f / VertexWeights.Num();
+		for (TTuple<FName, float>& Weight : VertexWeights)
+		{
+			Weight.Value = SplitWeight;
+		}
+		UE_LOG(LogAnimation, Warning, TEXT("Skin Weight Modifier: Supplied vertex had near zero total weight."));
+		return true;
+	}
+
+	// scale weights to achieve total value of 1.0
+	for (TTuple<FName, float>& Weight : VertexWeights)
+	{
+		Weight.Value /= TotalWeight;
+	}
+
+	return true;
+}
+
+void USkinWeightModifier::EnforceMaxInfluenceInternal(const int32 VertexID, const int32 Max)
+{
+	// get weights for this vertex
+	TMap<FName, float>& VertexWeights = Weights[VertexID];
+
+	// already under max influences?
+	if (VertexWeights.Num() <= Max)
+	{
+		return;
+	}
+
+	// remove smallest influences until we reach max influences
+	while(VertexWeights.Num() > Max)
+	{
+		FName SmallestInfName = NAME_None;
+		float SmallestInf = TNumericLimits<float>::Max();
+		for (const TTuple<FName, float> Weight : VertexWeights)
+		{
+			if (Weight.Value <= SmallestInf)
+			{
+				SmallestInf = Weight.Value;
+				SmallestInfName = Weight.Key;
+			}
+		}
+
+		VertexWeights.Remove(SmallestInfName);
+	}
+
+	NormalizeVertexInternal(VertexID);
+}
+
+bool USkinWeightModifier::PruneVertexWeightInternal(const int32 VertexID, const float WeightThreshold)
+{
+	// get weights for this vertex
+	TMap<FName, float>& VertexWeights = Weights[VertexID];
+	TArray<FName> InfluencesToPrune;
+	for (const TTuple<FName, float>& Weight : VertexWeights)
+	{
+		if (Weight.Value < WeightThreshold)
+		{
+			InfluencesToPrune.Add(Weight.Key);
+		}
+	}
+	
+	for (const FName InfluenceToPrune : InfluencesToPrune)
+	{
+		VertexWeights.Remove(InfluenceToPrune);
+	}
+	
+	return true;
+}
+
+void USkinWeightModifier::Reset()
+{
+	Mesh = nullptr;
+	MeshDescription.Reset();
+	Weights.Reset();
+}
+
+#undef LOCTEXT_NAMESPACE
\ No newline at end of file
diff --git a/Engine/Source/Editor/SkeletalMeshEditor/Public/SkinWeightModifier.h b/Engine/Source/Editor/SkeletalMeshEditor/Public/SkinWeightModifier.h
new file mode 100644
--- /dev/null
+++ b/Engine/Source/Editor/SkeletalMeshEditor/Public/SkinWeightModifier.h
@@ -0,0 +1,188 @@
+﻿// Copyright Epic Games, Inc. All Rights Reserved.
+
+#pragma once
+
+#include "BoneWeights.h"
+#include "Templates/UniquePtr.h"
+#include "MeshDescription.h"
+#include "Engine/SkeletalMesh.h"
+
+#include "SkinWeightModifier.generated.h"
+
+class USkeletalMesh;
+
+/**
+ * API used to modify skin weights on a Skeletal Mesh asset.
+ * 
+ * To use:
+ * 1. Instantiate an instance of USkinWeightModifier
+ * 2. Call "SetSkeletalMesh(MyMeshAsset)", passing in the skeletal mesh you want to edit weights on
+ * 3. Use Get/Set weights functions, and Normalize/Prune to edit the weights as desired
+ * 4. When ready to commit to the asset, call CommitWeightsToSkeletalMesh() and save the asset if desired
+ *
+ * This API can be used from C++, Blueprint or Python. Here is a sample usage of the API in Python:
+
+import unreal
+
+# create a weight modifier for a given skeletal mesh
+skel_mesh = unreal.EditorAssetLibrary().load_asset("/Game/Characters/Wolf/Meshes/SK_Wolf")
+weight_modifier = unreal.SkinWeightModifier()
+weight_modifier.set_skeletal_mesh(skel_mesh)
+
+# get weight of vertex 1234
+vertex_weights = weight_modifier.get_vertex_weights(1234)
+print(vertex_weights)
+
+# remove neck2 as an influence on this vertex
+vertex_weights.pop("neck2")
+weight_modifier.set_vertex_weights(vertex_weights, True)
+print(vertex_weights)
+
+# commit change to the skeletal mesh
+weight_modifier.commit_weights_to_skeletal_mesh()
+
+ * OUTPUT:
+ * {"head": 0.6, "neck1": 0.3, "neck2": 0.1}
+ * {"head": 0.6, "neck1": 0.3}
+ * 
+ * In Python, the per-vertex weights are stored as a dictionary mapping Bone Names to float weight values.
+ *
+ * The "SetVertexWeights()" function expects the same data structure. You can add/remove/edit influences as needed.
+ * The SetVertexWeights() function does not normalize the weights. So you can make multiple modifications
+ * and call NormalizeVertexWeights() or NormalizeAllWeights() as desired.
+ *
+ * The PruneVertexWeights() and EnforceMaxInfluences() functions can be used to trim small influences
+ * and clamp the total number of influences per vertex as needed.
+ * 
+ * Note that it is not required to normalize the weights by calling any of the normalize functions. Or manually before
+ * calling SetVertexWeights(). Committing the weights to the skeletal mesh will always enforce normalization.
+ *
+ * Though it may be useful to normalize while editing.
+ */
+UCLASS(BlueprintType)
+class SKELETALMESHEDITOR_API USkinWeightModifier : public UObject
+{
+	GENERATED_BODY()
+
+public:
+	
+	/**
+	 * Call this first to load the weights for a skeletal mesh for fast editing.
+	 * @param InMesh - The skeletal mesh asset to load for weight editing
+	 * @return bool - true if the mesh weights were successfully loaded.
+	 */
+	UFUNCTION(BlueprintCallable, Category = "Mesh")
+	bool SetSkeletalMesh(USkeletalMesh* InMesh);
+
+	/**
+	 * Actually applies the weight modifications to the skeletal mesh. This action creates an undo transaction.
+	 * The skeletal mesh asset will be dirtied, but it is up to the user to save the asset if required.
+	 * @return true if weights were applied to a skeletal mesh, false otherwise
+	 */
+	UFUNCTION(BlueprintCallable, Category = "Mesh")
+	bool CommitWeightsToSkeletalMesh();
+
+	/**
+	 * Get a reference to the skeletal mesh that was loaded
+	 * @return USkeletalMesh* - the mesh that was loaded, or null
+	 */
+	UFUNCTION(BlueprintCallable, Category = "Mesh")
+	USkeletalMesh* GetSkeletalMesh() { return Mesh; }
+
+	/**
+	 * Get the total number of vertices in the skeletal mesh.
+	 * @return int, number of vertices
+	 */
+	UFUNCTION(BlueprintCallable, Category = "Mesh Query")
+	int32 GetNumVertices();
+
+	/**
+	 * Get an array of all bone names in the skeletal mesh.
+	 * @return array of bone names
+	 */
+	UFUNCTION(BlueprintCallable, Category = "Mesh Query")
+	TArray<FName> GetAllBoneNames();
+
+	/**
+	 * Get all bone weights for a single vertex.
+	 * @param VertexID the index of the vertex
+	 * @return a map of Bone Name to weight values for all bones that influence the specified vertex.
+	 */
+	UFUNCTION(BlueprintCallable, Category = "Weights")
+	TMap<FName, float> GetVertexWeights(const int32 VertexID);
+
+	/**
+	 * Set bone weights for a single vertex. The weights are stored as supplied and not normalized until
+	 * either "CommitWeightsToSkeletalMesh()" or "NormalizeVertexWeights()" is called.
+	 * @param VertexID the index of the vertex
+	 * @param InWeights a map of Bone-Name to Weight for all bones that influence the specified vertex, ie {"Head": 0.6, "Neck": 0.4}
+	 * @param bReplaceAll if true, all weights on this vertex will be replaced with the input weights. Default is false.
+	 */
+	UFUNCTION(BlueprintCallable, Category = "Weights")
+	bool SetVertexWeights(
+		const int32 VertexID,
+		const TMap<FName, float>& InWeights,
+		const bool bReplaceAll=false);
+
+	/**
+	 * Normalize weights on the specified vertex.
+	 * @param VertexID the index of the vertex to normalize weights on
+	 * @return true if normalization was performed, false otherwise
+	 */
+	UFUNCTION(BlueprintCallable, Category = "Weights")
+	bool NormalizeVertexWeights(const int32 VertexID);
+
+	/**
+	 * Normalize weights on all vertices in the mesh.
+	 * @return true if normalization was performed, false otherwise
+	 */
+	UFUNCTION(BlueprintCallable, Category = "Weights")
+	bool NormalizeAllWeights();
+
+	/**
+	 * Strips out smallest influences to ensure each vertex does not have weight on more influences than MaxInfluences.
+	 * Influences with the smallest weight are culled first.
+	 * @param MaxInfluences The maximum number of influences to allow for each vertex in the mesh. If -1 is passed, will use the project-wide MaxInfluences amount.
+	 * @return true if influences were removed, false otherwise
+	 */
+	UFUNCTION(BlueprintCallable, Category = "Weights")
+	bool EnforceMaxInfluences(const int32 MaxInfluences=-1);
+
+	/**
+	 * Remove all weights below the given threshold value, on the given vertex.
+	 * Influences that are pruned will no longer receive weight from normalization.
+	 * @param VertexID the index of the vertex to prune weights on
+	 * @param WeightThreshold vertex weights below this value will be removed. 
+	 * @return true if influences were removed, false otherwise
+	 */
+	UFUNCTION(BlueprintCallable, Category = "Weights")
+	bool PruneVertexWeights(const int32 VertexID, const float WeightThreshold);
+
+	/**
+	 * Remove all weights below the given threshold value, on all vertices.
+	 * @param WeightThreshold vertex weights below this value will be removed. 
+	 * @return true if influences were removed, false otherwise
+	 */
+	UFUNCTION(BlueprintCallable, Category = "Weights")
+	bool PruneAllWeights(const float WeightThreshold);
+
+private:
+
+	// the skeletal mesh that was loaded by SetSkeletalMesh
+	UPROPERTY()
+	TObjectPtr<USkeletalMesh> Mesh;
+	TUniquePtr<FMeshDescription> MeshDescription;
+
+	// weights stored in-memory for editing before being committed to the skeletal mesh
+	TArray<TMap<FName, float>> Weights;
+
+	// fast internal procedures (do not validate inputs)
+	bool NormalizeVertexInternal(const int32 VertexID);
+	void EnforceMaxInfluenceInternal(const int32 VertexID, const int32 Max);
+	bool PruneVertexWeightInternal(const int32 VertexID, const float WeightThreshold);
+
+	// reset all internal data when loading a new mesh
+	void Reset();
+
+	static constexpr int32 LODIndex = 0;
+};
diff --git a/Engine/Source/Editor/SkeletalMeshEditor/SkeletalMeshEditor.Build.cs b/Engine/Source/Editor/SkeletalMeshEditor/SkeletalMeshEditor.Build.cs
--- a/Engine/Source/Editor/SkeletalMeshEditor/SkeletalMeshEditor.Build.cs
+++ b/Engine/Source/Editor/SkeletalMeshEditor/SkeletalMeshEditor.Build.cs
@@ -41,6 +41,9 @@ public class SkeletalMeshEditor : ModuleRules
 				"PhysicsUtilities",
 				"InterchangeCore",
 				"InterchangeEngine",
+				"MeshDescription",
+				"SkeletalMeshDescription",
+				"RenderCore"
 			}
 		);
 
