Do not set main tab for the skeleton editor
    - follow up to 25175110

diff --git a/Engine/Source/Editor/SkeletalMeshEditor/Private/SkeletalMeshEditor.cpp b/Engine/Source/Editor/SkeletalMeshEditor/Private/SkeletalMeshEditor.cpp
--- a/Engine/Source/Editor/SkeletalMeshEditor/Private/SkeletalMeshEditor.cpp
+++ b/Engine/Source/Editor/SkeletalMeshEditor/Private/SkeletalMeshEditor.cpp
@@ -261,6 +261,12 @@ void FSkeletalMeshEditor::InitSkeletalMeshEditor(const EToolkitMode::Type Mode,
 	const bool bCreateDefaultToolbar = true;
 	FAssetEditorToolkit::InitAssetEditor(Mode, InitToolkitHost, SkeletalMeshEditorAppIdentifier, FTabManager::FLayout::NullLayout, bCreateDefaultStandaloneMenu, bCreateDefaultToolbar, InSkeletalMesh);
 
+	if (TabManager)
+	{
+		TabManager->SetEnforceMainTab(false);
+		TabManager->SetMainTab(FTabId());
+	}
+
 	BindCommands();
 
 	AddApplicationMode(
