Fixed shared ptr use preventing the toolkit from being deleted (use weak ptr instead)
    + cleaned some useless code

diff --git a/Engine/Source/Editor/SkeletalMeshEditor/Private/SSkeletalMeshEditorToolbox.cpp b/Engine/Source/Editor/SkeletalMeshEditor/Private/SSkeletalMeshEditorToolbox.cpp
--- a/Engine/Source/Editor/SkeletalMeshEditor/Private/SSkeletalMeshEditorToolbox.cpp
+++ b/Engine/Source/Editor/SkeletalMeshEditor/Private/SSkeletalMeshEditorToolbox.cpp
@@ -107,7 +107,7 @@ void SSkeletalMeshEditorToolbox::UpdateInlineContent(const TSharedPtr<IToolkit>&
 		TabIcon = ModeToolkit.Pin()->GetEditorModeIcon().GetSmallIcon();
 		const TSharedRef<FModeToolkit> ModeToolkitPinned = ModeToolkit.Pin().ToSharedRef();
 
-		UpdatePalette(ModeToolkitPinned);
+		UpdatePalette(ModeToolkit);
 
 		// Show the name of the active tool in the statusbar.
 		// FIXME: We should also be showing Ctrl/Shift/Alt LMB/RMB shortcuts.
@@ -135,12 +135,26 @@ void SSkeletalMeshEditorToolbox::UpdateInlineContent(const TSharedPtr<IToolkit>&
 }
 
 
-void SSkeletalMeshEditorToolbox::UpdatePalette(const TSharedRef<FModeToolkit>& InModeToolkit)
+void SSkeletalMeshEditorToolbox::UpdatePalette(const TWeakPtr<FModeToolkit>& InWeakToolkit) const
 {
+	if (!InWeakToolkit.IsValid())
+	{
+		return;
+	}
+	
 	TSharedRef<SSegmentedControl<FName>> PaletteTabBox = SNew(SSegmentedControl<FName>)
 		.UniformPadding(FMargin(8.f, 3.f))
-		.Value_Lambda([InModeToolkit]() { return InModeToolkit->GetCurrentPalette(); })
-		.OnValueChanged_Lambda([InModeToolkit](const FName& Palette) { InModeToolkit->SetCurrentPalette(Palette); });
+		.Value_Lambda([InWeakToolkit]()
+		{
+			return InWeakToolkit.IsValid() ? InWeakToolkit.Pin()->GetCurrentPalette() : NAME_None;
+		})
+		.OnValueChanged_Lambda([InWeakToolkit](const FName& Palette)
+		{
+			if (InWeakToolkit.IsValid())
+			{
+				InWeakToolkit.Pin()->SetCurrentPalette(Palette);
+			}
+		});
 
 	// Only show if there's more than one entry.
 	PaletteTabBox->SetVisibility(TAttribute<EVisibility>::Create(
@@ -148,36 +162,42 @@ void SSkeletalMeshEditorToolbox::UpdatePalette(const TSharedRef<FModeToolkit>& I
 			return PaletteTabBox->NumSlots() > 1 ? EVisibility::Visible : EVisibility::Collapsed;
 		})));
 
+	const TSharedPtr<FModeToolkit> ModeToolkit = InWeakToolkit.Pin();
+	
 	// Also build the toolkit here
 	TArray<FName> PaletteNames;
-	InModeToolkit->GetToolPaletteNames(PaletteNames);
+	ModeToolkit->GetToolPaletteNames(PaletteNames);
 
-	TSharedPtr<FUICommandList> CommandList;
-	CommandList = InModeToolkit->GetToolkitCommands();
+	const TSharedPtr<FUICommandList> CommandList = ModeToolkit->GetToolkitCommands();
 
 	TSharedRef< SWidgetSwitcher > PaletteSwitcher = SNew(SWidgetSwitcher)
-	.WidgetIndex_Lambda( [PaletteNames, InModeToolkit] () -> int32 { 
-		int32 FoundIndex;
-		if (PaletteNames.Find(InModeToolkit->GetCurrentPalette(), FoundIndex))
+	.WidgetIndex_Lambda( [PaletteNames, InWeakToolkit] () -> int32 {
+		if (InWeakToolkit.IsValid())
 		{
-			return FoundIndex;	
+			int32 FoundIndex;
+			if (PaletteNames.Find(InWeakToolkit.Pin()->GetCurrentPalette(), FoundIndex))
+			{
+				return FoundIndex;	
+			}
 		}
 		return 0;
 	} );
 			
 	for(auto Palette : PaletteNames)
 	{
-		FName ToolbarCustomizationName = InModeToolkit->GetEditorMode() ? InModeToolkit->GetEditorMode()->GetModeInfo().ToolbarCustomizationName : InModeToolkit->GetScriptableEditorMode()->GetModeInfo().ToolbarCustomizationName;
+		const FName ToolbarCustomizationName = ModeToolkit->GetEditorMode() ?
+			ModeToolkit->GetEditorMode()->GetModeInfo().ToolbarCustomizationName :
+			ModeToolkit->GetScriptableEditorMode()->GetModeInfo().ToolbarCustomizationName;
 		FUniformToolBarBuilder ModeToolbarBuilder(CommandList, FMultiBoxCustomization(ToolbarCustomizationName));
 		ModeToolbarBuilder.SetStyle(&FAppStyle::Get(), "PaletteToolBar");
 
-		InModeToolkit->BuildToolPalette(Palette, ModeToolbarBuilder);
+		ModeToolkit->BuildToolPalette(Palette, ModeToolbarBuilder);
 
 		TSharedRef<SWidget> PaletteWidget = ModeToolbarBuilder.MakeWidget();
 
 		const bool bRebuildChildren = false;
 		PaletteTabBox->AddSlot(Palette, false)
-		.Text(InModeToolkit->GetToolPaletteDisplayName(Palette));
+		.Text(ModeToolkit->GetToolPaletteDisplayName(Palette));
 
 		PaletteSwitcher->AddSlot()
 		[
diff --git a/Engine/Source/Editor/SkeletalMeshEditor/Private/SSkeletalMeshEditorToolbox.h b/Engine/Source/Editor/SkeletalMeshEditor/Private/SSkeletalMeshEditorToolbox.h
--- a/Engine/Source/Editor/SkeletalMeshEditor/Private/SSkeletalMeshEditorToolbox.h
+++ b/Engine/Source/Editor/SkeletalMeshEditor/Private/SSkeletalMeshEditorToolbox.h
@@ -40,7 +40,7 @@ public:
 private:	
 	void UpdateInlineContent(const TSharedPtr<IToolkit>& Toolkit, TSharedPtr<SWidget> InlineContent);
 
-	void UpdatePalette(const TSharedRef<FModeToolkit>& InModeToolkit);
+	void UpdatePalette(const TWeakPtr<FModeToolkit>& InWeakToolkit) const;
 
 	// Whether to show the inline content or not.
 	EVisibility GetInlineContentHolderVisibility() const;
diff --git a/Engine/Source/Editor/SkeletalMeshEditor/Private/SkeletalMeshEditor.cpp b/Engine/Source/Editor/SkeletalMeshEditor/Private/SkeletalMeshEditor.cpp
--- a/Engine/Source/Editor/SkeletalMeshEditor/Private/SkeletalMeshEditor.cpp
+++ b/Engine/Source/Editor/SkeletalMeshEditor/Private/SkeletalMeshEditor.cpp
@@ -266,12 +266,6 @@ void FSkeletalMeshEditor::InitSkeletalMeshEditor(const EToolkitMode::Type Mode,
 	const bool bCreateDefaultToolbar = true;
 	FAssetEditorToolkit::InitAssetEditor(Mode, InitToolkitHost, SkeletalMeshEditorAppIdentifier, FTabManager::FLayout::NullLayout, bCreateDefaultStandaloneMenu, bCreateDefaultToolbar, InSkeletalMesh);
 
-	if (TabManager)
-	{
-		TabManager->SetEnforceMainTab(false);
-		TabManager->SetMainTab(FTabId());
-	}
-
 	BindCommands();
 
 	AddApplicationMode(
@@ -373,15 +367,6 @@ void FSkeletalMeshEditor::OnEditorModeIdChanged(const FEditorModeID& ModeChanged
 		// FIXME: We should get the hosted toolkit from here.
 		if (GetEditorModeManager().GetActiveScriptableMode(ModeChangedID)->UsesToolkits())
 		{
-			// for (const FTabId& TabId: DefaultTabs)
-			// {
-			// 	if (const TSharedPtr<SDockTab> Tab = TabManager->FindExistingLiveTab(TabId))
-			// 	{
-			// 		TabToRestore.Emplace(TabId, Tab->IsForeground());
-			// 		Tab->RequestCloseTab();
-			// 	}
-			// }
-
 			TabManager->TryInvokeTab(SkeletalMeshEditorTabs::ToolboxDetailsTab);
 		}
 	}
@@ -392,12 +377,6 @@ void FSkeletalMeshEditor::OnEditorModeIdChanged(const FEditorModeID& ModeChanged
 		{
 			ToolboxTab->RequestCloseTab();
 		}
-
-		// for (const TPair<FTabId, bool>& Tab: TabToRestore)
-		// {
-		// 	TabManager->TryInvokeTab(Tab.Get<0>());
-		// }
-		TabToRestore.Empty();
 	}
 }
 
diff --git a/Engine/Source/Editor/SkeletalMeshEditor/Private/SkeletalMeshEditor.h b/Engine/Source/Editor/SkeletalMeshEditor/Private/SkeletalMeshEditor.h
--- a/Engine/Source/Editor/SkeletalMeshEditor/Private/SkeletalMeshEditor.h
+++ b/Engine/Source/Editor/SkeletalMeshEditor/Private/SkeletalMeshEditor.h
@@ -221,9 +221,6 @@ private:
 
 	// Binding to send/receive skeletal mesh modifications
 	TSharedPtr<FSkeletalMeshEditorBinding> Binding;
-
-	// Tabs to restore when getting out of the tool
-	TArray<TPair<FTabId, bool>> TabToRestore;
 };
 
 /**
