Incrementing skeletal mesh editor tab layout version to force users onto new layout.

diff --git a/Engine/Source/Editor/SkeletalMeshEditor/Private/SkeletalMeshEditorMode.cpp b/Engine/Source/Editor/SkeletalMeshEditor/Private/SkeletalMeshEditorMode.cpp
--- a/Engine/Source/Editor/SkeletalMeshEditor/Private/SkeletalMeshEditorMode.cpp
+++ b/Engine/Source/Editor/SkeletalMeshEditor/Private/SkeletalMeshEditorMode.cpp
@@ -42,7 +42,7 @@ FSkeletalMeshEditorMode::FSkeletalMeshEditorMode(TSharedRef<FWorkflowCentricAppl
 	TabFactories.RegisterFactory(CreateMeshControllerMappingTabFactory(InHostingApp, Cast<USkeletalMesh> (SkeletalMeshEditor->HandleGetAsset()), SkeletalMeshEditor->OnPostUndo));
 	TabFactories.RegisterFactory(PersonaModule.CreateAnimAssetFindReplaceTabFactory(InHostingApp, FAnimAssetFindReplaceConfig()));
 
-	TabLayout = FTabManager::NewLayout("Standalone_SkeletalMeshEditor_Layout_v3.5")
+	TabLayout = FTabManager::NewLayout("Standalone_SkeletalMeshEditor_Layout_v4.5")
 		->AddArea
 		(
 			FTabManager::NewPrimaryArea()
