Fixed crash if attempt to paint weights on a socket.

diff --git a/Engine/Source/Editor/SkeletalMeshEditor/Private/SkeletalMeshEditor.cpp b/Engine/Source/Editor/SkeletalMeshEditor/Private/SkeletalMeshEditor.cpp
--- a/Engine/Source/Editor/SkeletalMeshEditor/Private/SkeletalMeshEditor.cpp
+++ b/Engine/Source/Editor/SkeletalMeshEditor/Private/SkeletalMeshEditor.cpp
@@ -1325,11 +1325,16 @@ void FSkeletalMeshEditor::HandleSelectionChanged(const TArrayView<TSharedPtr<ISk
 
 		if (Binding.IsValid())
 		{
-			TArray<FName> Selection;
-			Algo::TransformIf(InSelectedItems, Selection,
-				[](const TSharedPtr<ISkeletonTreeItem>& InItem) { return InItem->GetRowItemName() != NAME_None; },
-				[](const TSharedPtr<ISkeletonTreeItem>& InItem) { return InItem->GetRowItemName(); });
-			Binding->GetNotifier().Notify(Selection, ESkeletalMeshNotifyType::BonesSelected);
+			TArray<FName> BoneSelection;
+			Algo::TransformIf(InSelectedItems, BoneSelection,
+				[](const TSharedPtr<ISkeletonTreeItem>& InItem)
+				{
+					const bool bIsBoneType = InItem->IsOfTypeByName("FSkeletonTreeBoneItem");
+					const bool bIsNotNone = InItem->GetRowItemName() != NAME_None;
+					return bIsBoneType && bIsNotNone;
+				},
+				[](const TSharedPtr<ISkeletonTreeItem>& InItem){ return InItem->GetRowItemName(); });
+			Binding->GetNotifier().Notify(BoneSelection, ESkeletalMeshNotifyType::BonesSelected);
 		}
 	}
 }
