Viewport does not update in skeletal mesh editor until clicked on.

diff --git a/Engine/Source/Editor/SkeletalMeshEditor/Private/SkeletalMeshEditor.cpp b/Engine/Source/Editor/SkeletalMeshEditor/Private/SkeletalMeshEditor.cpp
--- a/Engine/Source/Editor/SkeletalMeshEditor/Private/SkeletalMeshEditor.cpp
+++ b/Engine/Source/Editor/SkeletalMeshEditor/Private/SkeletalMeshEditor.cpp
@@ -1373,6 +1373,10 @@ void FSkeletalMeshEditor::HandleMeshDetailsCreated(const TSharedRef<IDetailsView
 void FSkeletalMeshEditor::HandleViewportCreated(const TSharedRef<IPersonaViewport>& InViewport)
 {
 	Viewport = InViewport;
+	
+	// we need the viewport client to start out focused, or else it won't get ticked until we click inside it.
+	FEditorViewportClient& ViewportClient = InViewport->GetViewportClient(); 
+	ViewportClient.ReceivedFocus(ViewportClient.Viewport);
 }
 
 UObject* FSkeletalMeshEditor::HandleGetAsset()
