Fixed toolbox tab not opening

#rnx

diff --git a/Engine/Source/Editor/SkeletalMeshEditor/Private/SkeletalMeshEditor.cpp b/Engine/Source/Editor/SkeletalMeshEditor/Private/SkeletalMeshEditor.cpp
--- a/Engine/Source/Editor/SkeletalMeshEditor/Private/SkeletalMeshEditor.cpp
+++ b/Engine/Source/Editor/SkeletalMeshEditor/Private/SkeletalMeshEditor.cpp
@@ -75,7 +75,6 @@ const FName SkeletalMeshEditorTabs::AssetDetailsTab(TEXT("AnimAssetPropertiesTab
 const FName SkeletalMeshEditorTabs::ViewportTab(TEXT("Viewport"));
 const FName SkeletalMeshEditorTabs::AdvancedPreviewTab(TEXT("AdvancedPreviewTab"));
 const FName SkeletalMeshEditorTabs::MorphTargetsTab("MorphTargetsTab");
-const FName SkeletalMeshEditorTabs::ToolboxDetailsTab("ToolBoxDetailsTab");
 const FName SkeletalMeshEditorTabs::CurveMetadataTab(TEXT("AnimCurveMetadataEditorTab"));
 const FName SkeletalMeshEditorTabs::FindReplaceTab("FindReplaceTab");
 
@@ -263,10 +262,6 @@ void FSkeletalMeshEditor::InitSkeletalMeshEditor(const EToolkitMode::Type Mode,
 	PreviewScene->RegisterOnMeshClick(FOnMeshClick::CreateSP(this, &FSkeletalMeshEditor::HandleMeshClick));
 	PreviewScene->SetAllowMeshHitProxies(true);
 
-	// Make sure we get told when the editor mode changes so we can switch to the appropriate tab
-	// if there's a toolbox available.
-	GetEditorModeManager().OnEditorModeIDChanged().AddSP(this, &FSkeletalMeshEditor::OnEditorModeIdChanged);
-
 	// run attached post-init delegates
 	ISkeletalMeshEditorModule& SkeletalMeshEditorModule = FModuleManager::GetModuleChecked<ISkeletalMeshEditorModule>("SkeletalMeshEditor");
 	const TArray<ISkeletalMeshEditorModule::FOnSkeletalMeshEditorInitialized>& PostInitDelegates = SkeletalMeshEditorModule.GetPostEditorInitDelegates();
@@ -333,41 +328,6 @@ TSharedPtr<FSkeletalMeshEditor> FSkeletalMeshEditor::GetSkeletalMeshEditor(const
 	return TSharedPtr<FSkeletalMeshEditor>();
 }
 
-void FSkeletalMeshEditor::OnEditorModeIdChanged(const FEditorModeID& ModeChangedID, bool bIsEnteringMode)
-{
-	if (GetEditorModeManager().IsDefaultMode(ModeChangedID))
-	{
-		return;
-	}
-
-	static const TArray<FTabId> DefaultTabs{
-		SkeletalMeshEditorTabs::DetailsTab,
-		SkeletalMeshEditorTabs::AssetDetailsTab,
-		SkeletalMeshEditorTabs::SkeletonTreeTab,
-		SkeletalMeshEditorTabs::AdvancedPreviewTab,
-		SkeletalMeshEditorTabs::MorphTargetsTab,
-		SkeletalMeshEditorTabs::CurveMetadataTab,
-		SkeletalMeshEditorTabs::FindReplaceTab
-	};
-	
-	if (bIsEnteringMode)
-	{
-		// FIXME: We should get the hosted toolkit from here.
-		if (GetEditorModeManager().GetActiveScriptableMode(ModeChangedID)->UsesToolkits())
-		{
-			TabManager->TryInvokeTab(SkeletalMeshEditorTabs::ToolboxDetailsTab);
-		}
-	}
-	else
-	{
-		TSharedPtr<SDockTab> ToolboxTab = TabManager->FindExistingLiveTab(SkeletalMeshEditorTabs::ToolboxDetailsTab);
-		if (ToolboxTab.IsValid())
-		{
-			ToolboxTab->RequestCloseTab();
-		}
-	}
-}
-
 void FSkeletalMeshEditor::RegisterReimportContextMenu(const FName InBaseMenuName)
 {
 	static auto ReimportMeshWithNewFileAction = [](const FToolMenuContext& InMenuContext, int32 SourceFileIndex)
diff --git a/Engine/Source/Editor/SkeletalMeshEditor/Private/SkeletalMeshEditor.h b/Engine/Source/Editor/SkeletalMeshEditor/Private/SkeletalMeshEditor.h
--- a/Engine/Source/Editor/SkeletalMeshEditor/Private/SkeletalMeshEditor.h
+++ b/Engine/Source/Editor/SkeletalMeshEditor/Private/SkeletalMeshEditor.h
@@ -45,7 +45,6 @@ namespace SkeletalMeshEditorTabs
 	extern const FName MorphTargetsTab;
 	extern const FName MeshDetailsTab;
 	extern const FName AnimationMappingTab;
-	extern const FName ToolboxDetailsTab;
 	extern const FName CurveMetadataTab;
 	extern const FName FindReplaceTab;
 }
@@ -166,12 +165,6 @@ private:
 
 	static TSharedPtr<FSkeletalMeshEditor> GetSkeletalMeshEditor(const FToolMenuContext& InMenuContext);
 
-	void OnEditorModeIdChanged(const FEditorModeID& ModeChangedID, bool bIsEnteringMode);
-
-	bool CanSpawnToolboxTab(const FSpawnTabArgs& InArgs) const;
-	TSharedRef<SDockTab> SpawnToolboxTab(const FSpawnTabArgs& InArgs);
-	void OnToolboxTabClosed(TSharedRef<SDockTab> InClosedTab);
-
 private:
 	void ExtendMenu();
 
diff --git a/Engine/Source/Editor/SkeletalMeshEditor/Private/SkeletalMeshEditorMode.cpp b/Engine/Source/Editor/SkeletalMeshEditor/Private/SkeletalMeshEditorMode.cpp
--- a/Engine/Source/Editor/SkeletalMeshEditor/Private/SkeletalMeshEditorMode.cpp
+++ b/Engine/Source/Editor/SkeletalMeshEditor/Private/SkeletalMeshEditorMode.cpp
@@ -11,6 +11,7 @@
 #include "Styling/AppStyle.h"
 #include "Engine/SkeletalMesh.h"
 #include "AnimAssetFindReplace.h"
+#include "PersonaTabs.h"
 
 #define LOCTEXT_NAMESPACE "SkeletalMeshEditorMode"
 
@@ -43,7 +44,7 @@ FSkeletalMeshEditorMode::FSkeletalMeshEditorMode(TSharedRef<FWorkflowCentricAppl
 	TabFactories.RegisterFactory(PersonaModule.CreateAnimAssetFindReplaceTabFactory(InHostingApp, FAnimAssetFindReplaceConfig()));
 	TabFactories.RegisterFactory(PersonaModule.CreatePersonaToolboxTabFactory(SkeletalMeshEditor));
 
-	TabLayout = FTabManager::NewLayout("Standalone_SkeletalMeshEditor_Layout_v4.5")
+	TabLayout = FTabManager::NewLayout("Standalone_SkeletalMeshEditor_Layout_v5")
 		->AddArea
 		(
 			FTabManager::NewPrimaryArea()
@@ -58,8 +59,8 @@ FSkeletalMeshEditorMode::FSkeletalMeshEditorMode(TSharedRef<FWorkflowCentricAppl
 					FTabManager::NewStack()
 					->SetSizeCoefficient(0.2f)
 					->SetHideTabWell(true)
-					->AddTab(SkeletalMeshEditorTabs::ToolboxDetailsTab, ETabState::OpenedTab)
-					->SetForegroundTab(SkeletalMeshEditorTabs::ToolboxDetailsTab)
+					->AddTab(FPersonaTabs::ToolboxID, ETabState::OpenedTab)
+					->SetForegroundTab(FPersonaTabs::ToolboxID)
 				)
 				->Split
 				(
