Added an AssignPhysicsAsset node to the skel mesh editor subsystem to handle editor scripting of modifying the physics asset for a mesh. Also added validation so we check to make sure all the bodies and constraints in the incoming physics asset are linked to bones that actually exist on the target skeletal mesh. Validation also exposed as a seperate node.

diff --git a/Engine/Source/Editor/SkeletalMeshEditor/Private/SkeletalMeshEditorSubsystem.cpp b/Engine/Source/Editor/SkeletalMeshEditor/Private/SkeletalMeshEditorSubsystem.cpp
--- a/Engine/Source/Editor/SkeletalMeshEditor/Private/SkeletalMeshEditorSubsystem.cpp
+++ b/Engine/Source/Editor/SkeletalMeshEditor/Private/SkeletalMeshEditorSubsystem.cpp
@@ -16,6 +16,7 @@
 #include "ObjectTools.h"
 #include "PhysicsAssetUtils.h"
 #include "PhysicsEngine/PhysicsAsset.h"
+#include "PhysicsEngine/PhysicsConstraintTemplate.h"
 #include "Rendering/SkeletalMeshLODModel.h"
 #include "Rendering/SkeletalMeshLODRenderData.h"
 #include "Rendering/SkeletalMeshModel.h"
@@ -855,4 +856,87 @@ UPhysicsAsset* USkeletalMeshEditorSubsystem::CreatePhysicsAsset(USkeletalMesh* S
 	return NewPhysicsAsset;
 }
 
+bool USkeletalMeshEditorSubsystem::IsPhysicsAssetCompatible(USkeletalMesh* TargetMesh, UPhysicsAsset* PhysicsAsset)
+{
+	if(!TargetMesh)
+	{
+		return false;
+	}
+
+	// If we have a physics asset, validate it. Otherwise the user is clearing the asset by passing a nullptr
+	if (PhysicsAsset)
+	{
+		const FReferenceSkeleton& RefSkel = TargetMesh->GetRefSkeleton();
+
+		// Check that all of the body setups correctly map to a skeletal mesh bone
+		for (USkeletalBodySetup* Setup : PhysicsAsset->SkeletalBodySetups)
+		{
+			if (RefSkel.FindBoneIndex(Setup->BoneName) == INDEX_NONE)
+			{
+				UE_LOG(LogSkeletalMeshEditorSubsystem, 
+				       Error, 
+				       TEXT("IsPhysicsAssetCompatible failed: bone %s is not present for skeletal mesh %s"), 
+				       *Setup->BoneName.ToString(), 
+				       *TargetMesh->GetName());
+
+				return false;
+			}
+		}
+
+		// Check that all of the constraints match up to existing skeletal mesh bones.
+		for (UPhysicsConstraintTemplate* Constraint : PhysicsAsset->ConstraintSetup)
+		{
+			int32 ConstraintBoneIndices[2] = {
+				RefSkel.FindBoneIndex(Constraint->DefaultInstance.ConstraintBone1),
+				RefSkel.FindBoneIndex(Constraint->DefaultInstance.ConstraintBone2)
+			};
+
+
+
+			if (ConstraintBoneIndices[0] == INDEX_NONE)
+			{
+				UE_LOG(LogSkeletalMeshEditorSubsystem, 
+				       Error, 
+				       TEXT("IsPhysicsAssetCompatible failed: bone %s is not present for skeletal mesh %s"), 
+				       *Constraint->DefaultInstance.ConstraintBone1.ToString(), 
+				       *TargetMesh->GetName());
+
+				return false;
+			}
+
+			if (ConstraintBoneIndices[1] == INDEX_NONE)
+			{
+				UE_LOG(LogSkeletalMeshEditorSubsystem, 
+				       Error, 
+				       TEXT("IsPhysicsAssetCompatible failed: bone %s is not present for skeletal mesh %s"), 
+				       *Constraint->DefaultInstance.ConstraintBone2.ToString(), 
+				       *TargetMesh->GetName());
+
+				return false;
+			}
+		}
+	}
+
+	return true;
+}
+
+bool USkeletalMeshEditorSubsystem::AssignPhysicsAsset(USkeletalMesh* TargetMesh, UPhysicsAsset* PhysicsAsset)
+{
+	if(!IsPhysicsAssetCompatible(TargetMesh, PhysicsAsset))
+	{
+		return false;
+	}
+
+	{
+		FScopedTransaction Transaction(LOCTEXT("SetSkeletalMeshPhysicsAsset", "Set skeletal mesh physics asset"));
+		TargetMesh->Modify();
+		TargetMesh->SetPhysicsAsset(PhysicsAsset);
+	}
+
+	// Refresh all currently active skeletal mesh components that may have used the old setup
+	RefreshSkelMeshOnPhysicsAssetChange(TargetMesh);
+
+	return true;
+}
+
 #undef LOCTEXT_NAMESPACE
diff --git a/Engine/Source/Editor/SkeletalMeshEditor/Public/SkeletalMeshEditorSubsystem.h b/Engine/Source/Editor/SkeletalMeshEditor/Public/SkeletalMeshEditorSubsystem.h
--- a/Engine/Source/Editor/SkeletalMeshEditor/Public/SkeletalMeshEditorSubsystem.h
+++ b/Engine/Source/Editor/SkeletalMeshEditor/Public/SkeletalMeshEditorSubsystem.h
@@ -218,4 +218,28 @@ public:
 	 */
 	UFUNCTION(BlueprintCallable, Category = "Editor Scripting | SkeletalMesh")
 	static UPhysicsAsset* CreatePhysicsAsset(USkeletalMesh* SkeletalMesh);
+
+	/**
+	 * Checks whether a physics asset is compatible with the given SkeletalMesh
+	 *
+	 * @param TargetMesh The mesh to test for compatibility
+	 * @param PhysicsAsset The PhysicsAsset to test for compatibility
+	 *
+	 * @return Whether the physics asset is compatible with the target SkeletalMesh
+	 */
+	UFUNCTION(BlueprintCallable, Category = "Editor Scripting | SkeletalMesh")
+	static bool IsPhysicsAssetCompatible(USkeletalMesh* TargetMesh, UPhysicsAsset* PhysicsAsset);
+
+	/**
+	 * Assigns a PhysicsAsset to the given SkeletalMesh if it is compatible. Passing
+	 * nullptr / None as the physics asset will always succeed and will clear the
+	 * physics asset assignment for the target SkeletalMesh
+	 *
+	 * @param TargetMesh The mesh to attempt to assign the PhysicsAsset to
+	 * @param PhysicsAsset The physics asset to assign to the provided mesh (or nullptr/None)
+	 *
+	 * @return Whether the physics asset was successfully assigned to the mesh
+	 */
+	UFUNCTION(BlueprintCallable, Category = "Editor Scripting | SkeletalMesh")
+	static bool AssignPhysicsAsset(USkeletalMesh* TargetMesh, UPhysicsAsset* PhysicsAsset);
 };
