Interchange: Fix skeletal mesh reimport: custom LODs, Alternate skin profiles and Cloth binding are now re-importing properly.
- Also fix the legacy alternate skinning if there is custom LODs in the asset. The alternate skin profiles are not lost anymore

#rnx

diff --git a/Engine/Plugins/Interchange/Runtime/Source/Import/InterchangeImport.Build.cs b/Engine/Plugins/Interchange/Runtime/Source/Import/InterchangeImport.Build.cs
--- a/Engine/Plugins/Interchange/Runtime/Source/Import/InterchangeImport.Build.cs
+++ b/Engine/Plugins/Interchange/Runtime/Source/Import/InterchangeImport.Build.cs
@@ -9,6 +9,7 @@ namespace UnrealBuildTool.Rules
 			PublicDependencyModuleNames.AddRange(
 				new string[]
 				{
+					"ClothingSystemRuntimeCommon",
 					"Core",
 					"CoreUObject",
 					"Engine",
diff --git a/Engine/Plugins/Interchange/Runtime/Source/Import/Private/Mesh/InterchangeSkeletalMeshFactory.cpp b/Engine/Plugins/Interchange/Runtime/Source/Import/Private/Mesh/InterchangeSkeletalMeshFactory.cpp
--- a/Engine/Plugins/Interchange/Runtime/Source/Import/Private/Mesh/InterchangeSkeletalMeshFactory.cpp
+++ b/Engine/Plugins/Interchange/Runtime/Source/Import/Private/Mesh/InterchangeSkeletalMeshFactory.cpp
@@ -16,6 +16,7 @@
 #include "InterchangeCommonPipelineDataFactoryNode.h"
 #include "InterchangeImportCommon.h"
 #include "InterchangeImportLog.h"
+#include "InterchangeManager.h"
 #include "InterchangeMaterialFactoryNode.h"
 #include "InterchangeMeshNode.h"
 #include "InterchangeSceneNode.h"
@@ -31,7 +32,9 @@
 #include "Mesh/InterchangeMeshHelper.h"
 #include "Mesh/InterchangeMeshPayload.h"
 #include "Mesh/InterchangeMeshPayloadInterface.h"
+#include "Misc/App.h"
 #include "Misc/MessageDialog.h"
+#include "Misc/ScopedSlowTask.h"
 #include "Nodes/InterchangeBaseNode.h"
 #include "Nodes/InterchangeBaseNodeContainer.h"
 #include "PhysicsEngine/PhysicsAsset.h"
@@ -44,6 +47,7 @@
 
 #if WITH_EDITOR
 #include "LODUtilities.h"
+#include "SkinWeightsUtilities.h"
 #endif
 
 #include UE_INLINE_GENERATED_CPP_BY_NAME(InterchangeSkeletalMeshFactory)
@@ -930,6 +934,68 @@ namespace UE
 
 #endif //#if WITH_EDITOR
 
+void FInterchangeSkeletalMeshPostImportTask::Execute()
+{
+#if WITH_EDITOR
+	if (!SkeletalMesh)
+	{
+		return;
+	}
+
+	//This code works only on the game thread and is not asynchronous
+	check(IsInGameThread());
+
+	if (bReImportAlternateSkinWeights)
+	{
+		//We can't show user if unattended just return;
+		if (GIsAutomationTesting || GIsRunningUnattendedScript || FApp::IsUnattended())
+		{
+			return;
+		}
+
+		FText MessageText = FText::Format(NSLOCTEXT("UInterchangeSkeletalMeshPostImportTask", "ShouldReimportAlternateSkinning", "When reimporting asset \"{0}\" the system cannot re-apply the alternate skinning data, do you want to re-import all alternate skinning?")
+			, FText::FromString(SkeletalMesh->GetName()));
+
+		EAppReturnType::Type ReimportAlternateSkinningChoice = FMessageDialog::Open(EAppMsgType::YesNo
+			, EAppReturnType::No
+			, MessageText);
+		if (ReimportAlternateSkinningChoice == EAppReturnType::No)
+		{
+			return;
+		}
+
+		//User say yes so re-import the alternate skinning
+		const int32 LodCount = SkeletalMesh->GetLODNum();
+		float ProgressCount = LodCount + 0.1f;
+
+		FScopedSlowTask Progress(ProgressCount, NSLOCTEXT("UInterchangeSkeletalMeshPostImportTask", "SkeletalMeshPostImportTaskGameThread", "Executing Skeletal Mesh Post Import Tasks..."));
+		Progress.MakeDialog();
+		{
+			//Make sure we rebuild the skeletal mesh after re-importing all skin weight
+			FScopedSkeletalMeshPostEditChange ScopePostEditChange(SkeletalMesh);
+
+			//Wait until the asset is finish building then lock the skeletal mesh properties to prevent the UI to update during the alternate skinning reimport
+			FEvent* LockEvent = SkeletalMesh->LockPropertiesUntil();
+			FSkinnedAssetAsyncBuildScope AsyncBuildScope(SkeletalMesh);
+
+			//We have a 0.1 progress for the lock
+			Progress.EnterProgressFrame(0.1f);
+
+			//Reimport all the alternate skinning
+			for (int32 LodIndex = 0; LodIndex < LodCount; ++LodIndex)
+			{
+				FSkinWeightsUtilities::ReimportAlternateSkinWeight(SkeletalMesh, LodIndex);
+				Progress.EnterProgressFrame(1.0f);
+			}
+
+			//Release the skeletal mesh async properties
+			LockEvent->Trigger();
+
+			//Skeletal mesh will rebuild when going out of scope
+		}
+	}
+#endif //WITH_EDITOR
+}
 
 UClass* UInterchangeSkeletalMeshFactory::GetFactoryClass() const
 {
@@ -1011,6 +1077,37 @@ UInterchangeFactoryBase::FImportAssetResult UInterchangeSkeletalMeshFactory::Beg
 	SkeletalMesh->InvalidateDeriveDataCacheGUID();
 	USkeleton* SkeletonReference = nullptr;
 
+	if (bIsReImport)
+	{
+		//Save all existing source data that are imported only by the editor UI
+
+		int32 ExistingLodCount = SkeletalMesh->GetLODNum();
+
+		//Skin weight profiles, The skin weight alternate data will be extracted when iterating the lods
+		ImportAssetObjectData.ExistingSkinWeightProfileInfos = SkeletalMesh->GetSkinWeightProfiles();
+
+		//Unbind clothing and save the data to rebind it later in the post import task
+		SkeletalMesh->GetSkinWeightProfiles().Reset();
+		for (int32 LodIndex = 0; LodIndex < ExistingLodCount; ++LodIndex)
+		{
+			FSkeletalMeshLODModel& BuildLODModel = SkeletalMesh->GetImportedModel()->LODModels[LodIndex];
+			BuildLODModel.SkinWeightProfiles.Reset();
+
+			//Store the LOD alternate skinning profile data
+			FSkeletalMeshImportData& SkeletalMeshImportData = ImportAssetObjectData.ExistingAlternateImportDataPerLOD.AddDefaulted_GetRef();
+			if (SkeletalMesh->HasMeshDescription(LodIndex))
+			{
+				PRAGMA_DISABLE_DEPRECATION_WARNINGS
+				SkeletalMesh->LoadLODImportedData(LodIndex, SkeletalMeshImportData);
+				PRAGMA_ENABLE_DEPRECATION_WARNINGS
+			}
+
+			//Cloth
+			TArray<ClothingAssetUtils::FClothingAssetMeshBinding> ExistingClothingBindingsLod;
+			FLODUtilities::UnbindClothingAndBackup(SkeletalMesh, ExistingClothingBindingsLod, LodIndex);
+			ImportAssetObjectData.ExistingClothingBindings.Append(ExistingClothingBindingsLod);
+		}
+	}
 
 	int32 LodCount = SkeletalMeshFactoryNode->GetLodDataCount();
 	TArray<FString> LodDataUniqueIds;
@@ -1955,8 +2052,142 @@ void UInterchangeSkeletalMeshFactory::SetupObject_GameThread(const FSetupObjectP
 
 		
 		SkeletalMesh->SetAssetImportData(ImportDataPtr);
+#if WITH_EDITOR
+		//Re-apply the alternate skinning data
+		if (ImportAssetObjectData.ExistingSkinWeightProfileInfos.Num() > 0)
+		{
+			TSharedPtr<FInterchangeSkeletalMeshPostImportTask> SkeletalMeshPostImportTask = nullptr;
+			TArray<FSkinWeightProfileInfo>& SkinProfiles = SkeletalMesh->GetSkinWeightProfiles();
+			SkinProfiles = ImportAssetObjectData.ExistingSkinWeightProfileInfos;
+			for (const FSkinWeightProfileInfo& ProfileInfo : SkinProfiles)
+			{
+				const int32 LodCount = SkeletalMesh->GetLODNum();
+				for (int32 LodIndex = 0; LodIndex < LodCount; ++LodIndex)
+				{
+					if (!ImportAssetObjectData.ExistingAlternateImportDataPerLOD.IsValidIndex(LodIndex))
+					{
+						continue;
+					}
+					if (!SkeletalMesh->HasMeshDescription(LodIndex))
+					{
+						continue;
+					}
+					const FSkeletalMeshLODInfo* LodInfo = SkeletalMesh->GetLODInfo(LodIndex);
+					if (!LodInfo)
+					{
+						continue;
+					}
+
+					const FSkeletalMeshImportData& ExistingImportDataSrc = ImportAssetObjectData.ExistingAlternateImportDataPerLOD[LodIndex];
+
+					const FString ProfileNameStr = ProfileInfo.Name.ToString();
+
+					FSkeletalMeshImportData ImportDataDest;
+					PRAGMA_DISABLE_DEPRECATION_WARNINGS
+					SkeletalMesh->LoadLODImportedData(LodIndex, ImportDataDest);
+					PRAGMA_ENABLE_DEPRECATION_WARNINGS
+					
+					int32 PointNumberDest = ImportDataDest.Points.Num();
+					int32 VertexNumberDest = ImportDataDest.Points.Num();
+
+					if (ExistingImportDataSrc.Points.Num() != PointNumberDest)
+					{
+						//Warn the user but still apply the data
+						UInterchangeResultWarning_Generic* Message = AddMessage<UInterchangeResultWarning_Generic>();
+						Message->Text = FText::Format(NSLOCTEXT("InterchangeSkeletalMeshFactory", "SetupObject_GameThread_AlternateSkinning missmatch", "Alternate skinning mesh vertex number is different from the mesh LOD, we cannot apply the existing alternate skinning [{0}] when re-importing skeletal mesh LOD{1} [{2}]")
+							, FText::FromString(ProfileNameStr)
+							, FText::AsNumber(LodIndex)
+							, FText::FromString(SkeletalMesh->GetName()));
+
+						//We must enqueue a post import task that will re-import all skin weight profile.
+						if (!SkeletalMeshPostImportTask.IsValid())
+						{
+							SkeletalMeshPostImportTask = MakeShared<FInterchangeSkeletalMeshPostImportTask>();
+							SkeletalMeshPostImportTask->SkeletalMesh = SkeletalMesh;
+							SkeletalMeshPostImportTask->bReImportAlternateSkinWeights = true;
+							UInterchangeManager::GetInterchangeManager().EnqueuePostImportTask(SkeletalMeshPostImportTask);
+						}
+						continue;
+					}
+
+					//Replace the data into the destination bulk data and save it
+					int32 ProfileIndex = 0;
+					if (ImportDataDest.AlternateInfluenceProfileNames.Find(ProfileNameStr, ProfileIndex))
+					{
+						ImportDataDest.AlternateInfluenceProfileNames.RemoveAt(ProfileIndex);
+						ImportDataDest.AlternateInfluences.RemoveAt(ProfileIndex);
+					}
+					int32 SrcProfileIndex = 0;
+					if (ExistingImportDataSrc.AlternateInfluenceProfileNames.Find(ProfileNameStr, SrcProfileIndex))
+					{
+						ImportDataDest.AlternateInfluenceProfileNames.Add(ProfileNameStr);
+						ImportDataDest.AlternateInfluences.Add(ExistingImportDataSrc.AlternateInfluences[SrcProfileIndex]);
+					}
+
+					//Resave the bulk data with the new or refreshed data
+					PRAGMA_DISABLE_DEPRECATION_WARNINGS
+					SkeletalMesh->SaveLODImportedData(LodIndex, ImportDataDest);
+					PRAGMA_ENABLE_DEPRECATION_WARNINGS
+				}
+			}
+		}
+#endif //WITH_EDITOR
 	}
-#endif
+#endif //WITH_EDITORONLY_DATA
+}
+
+void UInterchangeSkeletalMeshFactory::FinalizeObject_GameThread(const FSetupObjectParams& Arguments)
+{
+	TRACE_CPUPROFILER_EVENT_SCOPE(UInterchangeSkeletalMeshFactory::FinalizeObject_GameThread)
+	Super::FinalizeObject_GameThread(Arguments);
+
+#if WITH_EDITOR
+	//This code works only on the game thread and is not asynchronous
+	check(IsInGameThread());
+
+	if (!ensure(Arguments.ImportedObject && Arguments.SourceData))
+	{
+		return;
+	}
+
+	USkeletalMesh* SkeletalMesh = CastChecked<USkeletalMesh>(Arguments.ImportedObject);
+
+	//Rebinding cloth will rebuild the skeletalmesh
+	//TODO: To avoid a second build we need to restore the clothing data before in SetupObject_GameThread.
+	//to do this we need to put the cloth binding data in the import data use by the build (i.e. Meshdescription)
+	if (ImportAssetObjectData.ExistingClothingBindings.Num() > 0)
+	{
+		float ProgressCount = 2.0f;
+
+		FScopedSlowTask Progress(ProgressCount, NSLOCTEXT("UInterchangeSkeletalMeshPostImportTask", "SkeletalMeshFinalizeImportGameThread", "Executing Skeletal Mesh Finalize Import Task..."));
+		Progress.MakeDialog();
+
+		//Make sure we rebuild the skeletal mesh after re-importing all skin weight
+		FScopedSkeletalMeshPostEditChange ScopePostEditChange(SkeletalMesh);
+
+		//Wait until the asset is finish building then lock the skeletal mesh properties to prevent the UI to update during the alternate skinning reimport
+		FEvent* LockEvent = SkeletalMesh->LockPropertiesUntil();
+
+		Progress.EnterProgressFrame(1.0f);
+
+		FSkinnedAssetAsyncBuildScope AsyncBuildScope(SkeletalMesh);
+
+		//Restore the clothing
+		if (ImportAssetObjectData.ExistingClothingBindings.Num() > 0)
+		{
+			FSkeletalMeshModel* ImportedResource = SkeletalMesh->GetImportedModel();
+			for (int32 LodIndex = 0; LodIndex < ImportedResource->LODModels.Num(); ++LodIndex)
+			{
+				// Re-apply our clothing assets
+				FLODUtilities::RestoreClothingFromBackup(SkeletalMesh, ImportAssetObjectData.ExistingClothingBindings, LodIndex);
+			}
+			Progress.EnterProgressFrame(1.0f);
+		}
+
+		//Release the skeletal mesh async properties
+		LockEvent->Trigger();
+	}
+#endif //WITH_EDITOR
 }
 
 bool UInterchangeSkeletalMeshFactory::GetSourceFilenames(const UObject* Object, TArray<FString>& OutSourceFilenames) const
diff --git a/Engine/Plugins/Interchange/Runtime/Source/Import/Public/Mesh/InterchangeSkeletalMeshFactory.h b/Engine/Plugins/Interchange/Runtime/Source/Import/Public/Mesh/InterchangeSkeletalMeshFactory.h
--- a/Engine/Plugins/Interchange/Runtime/Source/Import/Public/Mesh/InterchangeSkeletalMeshFactory.h
+++ b/Engine/Plugins/Interchange/Runtime/Source/Import/Public/Mesh/InterchangeSkeletalMeshFactory.h
@@ -2,6 +2,8 @@
 
 #pragma once
 
+#include "Animation/SkinWeightProfile.h"
+#include "ClothingAsset.h"
 #include "CoreMinimal.h"
 #include "InterchangeFactoryBase.h"
 #include "Rendering/SkeletalMeshLODImporterData.h"
@@ -14,7 +16,16 @@ class UInterchangeSkeletalMeshFactoryNode;
 class USkeletalMesh;
 class USkeleton;
 
-UCLASS(BlueprintType, Experimental)
+class INTERCHANGEIMPORT_API FInterchangeSkeletalMeshPostImportTask : public FInterchangePostImportTask
+{
+public:
+	virtual void Execute() override;
+
+	TObjectPtr<USkeletalMesh> SkeletalMesh;
+	bool bReImportAlternateSkinWeights = false;
+};
+
+UCLASS(BlueprintType)
 class INTERCHANGEIMPORT_API UInterchangeSkeletalMeshFactory : public UInterchangeFactoryBase
 {
 	GENERATED_BODY()
@@ -38,6 +49,12 @@ public:
 		bool bApplyGeometryOnly = false;
 		TArray<FImportAssetObjectLODData> LodDatas;
 
+		TArray<FSkinWeightProfileInfo> ExistingSkinWeightProfileInfos;
+		TArray<ClothingAssetUtils::FClothingAssetMeshBinding> ExistingClothingBindings;
+#if WITH_EDITOR
+		TArray<FSkeletalMeshImportData> ExistingAlternateImportDataPerLOD;
+#endif
+
 		bool IsValid() const;
 	};
 	//////////////////////////////////////////////////////////////////////////
@@ -50,6 +67,7 @@ public:
 	virtual FImportAssetResult EndImportAsset_GameThread(const FImportAssetObjectParams& Arguments) override;
 	virtual void Cancel() override;
 	virtual void SetupObject_GameThread(const FSetupObjectParams& Arguments) override;
+	virtual void FinalizeObject_GameThread(const FSetupObjectParams& Arguments) override;
 	virtual bool GetSourceFilenames(const UObject* Object, TArray<FString>& OutSourceFilenames) const override;
 	virtual bool SetSourceFilename(const UObject* Object, const FString& SourceFilename, int32 SourceIndex) const override;
 	virtual bool SetReimportSourceIndex(const UObject* Object, int32 SourceIndex) const override;
diff --git a/Engine/Plugins/Interchange/Runtime/Source/Pipelines/Private/InterchangeGenericAnimationPipeline.cpp b/Engine/Plugins/Interchange/Runtime/Source/Pipelines/Private/InterchangeGenericAnimationPipeline.cpp
--- a/Engine/Plugins/Interchange/Runtime/Source/Pipelines/Private/InterchangeGenericAnimationPipeline.cpp
+++ b/Engine/Plugins/Interchange/Runtime/Source/Pipelines/Private/InterchangeGenericAnimationPipeline.cpp
@@ -52,6 +52,7 @@ void UInterchangeGenericAnimationPipeline::AdjustSettingsForContext(EInterchange
 		|| ImportType == EInterchangePipelineContext::AssetAlternateSkinningReimport)
 	{
 		bImportAnimations = false;
+		CommonSkeletalMeshesAndAnimationsProperties->Skeleton = nullptr;
 		CommonSkeletalMeshesAndAnimationsProperties->bImportOnlyAnimations = false;
 	}
 	
diff --git a/Engine/Plugins/Interchange/Runtime/Source/Pipelines/Private/InterchangeGenericMaterialPipeline.cpp b/Engine/Plugins/Interchange/Runtime/Source/Pipelines/Private/InterchangeGenericMaterialPipeline.cpp
--- a/Engine/Plugins/Interchange/Runtime/Source/Pipelines/Private/InterchangeGenericMaterialPipeline.cpp
+++ b/Engine/Plugins/Interchange/Runtime/Source/Pipelines/Private/InterchangeGenericMaterialPipeline.cpp
@@ -620,6 +620,7 @@ void UInterchangeGenericMaterialPipeline::AdjustSettingsForContext(EInterchangeP
 	{
 		bImportMaterials = false;
 		HideCategories.Add(TEXT("Materials"));
+		SearchLocation = EInterchangeMaterialSearchLocation::DoNotSearch;
 	}
 
 	if (UInterchangePipelineBase* OuterMostPipeline = GetMostPipelineOuter())
diff --git a/Engine/Plugins/Interchange/Runtime/Source/Pipelines/Private/InterchangeGenericMeshPipeline.cpp b/Engine/Plugins/Interchange/Runtime/Source/Pipelines/Private/InterchangeGenericMeshPipeline.cpp
--- a/Engine/Plugins/Interchange/Runtime/Source/Pipelines/Private/InterchangeGenericMeshPipeline.cpp
+++ b/Engine/Plugins/Interchange/Runtime/Source/Pipelines/Private/InterchangeGenericMeshPipeline.cpp
@@ -43,6 +43,26 @@ void UInterchangeGenericMeshPipeline::AdjustSettingsForContext(EInterchangePipel
 	{
 		bCreatePhysicsAsset = false;
 		PhysicsAsset = nullptr;
+		
+		if (ImportType == EInterchangePipelineContext::AssetAlternateSkinningImport
+			|| ImportType == EInterchangePipelineContext::AssetAlternateSkinningReimport)
+		{
+			CommonMeshesProperties->ForceAllMeshAsType = EInterchangeForceMeshType::IFMT_SkeletalMesh;
+			CommonMeshesProperties->bAutoDetectMeshType = false;
+			CommonMeshesProperties->bBakeMeshes = true;
+			CommonMeshesProperties->bImportLods = false;
+			CommonMeshesProperties->bKeepSectionsSeparate = false;
+			CommonMeshesProperties->VertexColorImportOption = EInterchangeVertexColorImportOption::IVCIO_Ignore;
+			bImportSkeletalMeshes = true;
+			bImportStaticMeshes = false;
+			bBuildNanite = false;
+			bImportMorphTargets = false;
+			bImportVertexAttributes = false;
+			bUpdateSkeletonReferencePose = false;
+			SkeletalMeshImportContentType = EInterchangeSkeletalMeshContentType::All;
+			CommonSkeletalMeshesAndAnimationsProperties->Skeleton = nullptr;
+			CommonSkeletalMeshesAndAnimationsProperties->bImportOnlyAnimations = false;
+		}
 	}
 	const FString CommonMeshesCategory = TEXT("Common Meshes");
 	const FString StaticMeshesCategory = TEXT("Static Meshes");
diff --git a/Engine/Source/Developer/MeshBuilder/Private/SkeletalMeshBuilder.cpp b/Engine/Source/Developer/MeshBuilder/Private/SkeletalMeshBuilder.cpp
--- a/Engine/Source/Developer/MeshBuilder/Private/SkeletalMeshBuilder.cpp
+++ b/Engine/Source/Developer/MeshBuilder/Private/SkeletalMeshBuilder.cpp
@@ -158,20 +158,11 @@ bool FSkeletalMeshBuilder::Build(const FSkeletalMeshBuildParameters& SkeletalMes
 
 		//Re-apply the alternate skinning it must be after the inline reduction
 		SlowTask.EnterProgressFrame(1.0f, NSLOCTEXT("SkeltalMeshBuilder", "RebuildAlternateSkinning", "Rebuilding alternate skinning..."));
-		TArray<FSkinWeightProfileInfo> SkinProfiles = SkeletalMesh->GetSkinWeightProfiles();
+		const TArray<FSkinWeightProfileInfo> SkinProfiles = SkeletalMesh->GetSkinWeightProfiles();
 		for (int32 SkinProfileIndex = 0; SkinProfileIndex < SkinProfiles.Num(); ++SkinProfileIndex)
 		{
 			const FSkinWeightProfileInfo& ProfileInfo = SkinProfiles[SkinProfileIndex];
-			if (!FLODUtilities::UpdateAlternateSkinWeights(SkeletalMesh, ProfileInfo.Name, LODIndex, Options))
-			{
-				// If we failed to update the profile, remove it from the skeletal mesh so that we don't try to rely on it being there
-				// later on in the mesh build.
-				SkeletalMesh->GetSkinWeightProfiles().RemoveAll([Name=ProfileInfo.Name](const FSkinWeightProfileInfo& InProfileInfo)
-				{
-					return InProfileInfo.Name == Name;
-				});
-				BuildLODModel.SkinWeightProfiles.Remove(ProfileInfo.Name);
-			}
+			FLODUtilities::UpdateAlternateSkinWeights(SkeletalMesh, ProfileInfo.Name, LODIndex, Options);
 		}
 
 		// Copy vertex attribute definitions and their values from the import model.
diff --git a/Engine/Source/Developer/SkeletalMeshUtilitiesCommon/Private/LODUtilities.cpp b/Engine/Source/Developer/SkeletalMeshUtilitiesCommon/Private/LODUtilities.cpp
--- a/Engine/Source/Developer/SkeletalMeshUtilitiesCommon/Private/LODUtilities.cpp
+++ b/Engine/Source/Developer/SkeletalMeshUtilitiesCommon/Private/LODUtilities.cpp
@@ -509,6 +509,27 @@ void FLODUtilities::RemoveLODs(FSkeletalMeshUpdateContext& UpdateContext, const
 	}
 }
 
+namespace UE::Private
+{
+	struct FScopedLockProperties
+	{
+	public:
+		FScopedLockProperties(USkeletalMesh* SkeletalMesh)
+		{
+			check(SkeletalMesh);
+			Lock = SkeletalMesh->LockPropertiesUntil();
+		}
+
+		~FScopedLockProperties()
+		{
+			Lock->Trigger();
+			Lock = nullptr;
+		}
+	private:
+		FEvent* Lock = nullptr;
+	};
+}
+
 bool FLODUtilities::SetCustomLOD(USkeletalMesh* DestinationSkeletalMesh, USkeletalMesh* SourceSkeletalMesh, const int32 LodIndex, const FString& SourceDataFilename)
 {
 	if(!DestinationSkeletalMesh || !SourceSkeletalMesh)
@@ -517,9 +538,9 @@ bool FLODUtilities::SetCustomLOD(USkeletalMesh* DestinationSkeletalMesh, USkelet
 	}
 
 	FScopedSkeletalMeshPostEditChange ScopePostEditChange(DestinationSkeletalMesh);
-
-	//If the imported LOD already exist, we will need to reimport all the skin weight profiles
-	bool bMustReimportAlternateSkinWeightProfile = false;
+	//Lock the skeletal mesh
+	UE::Private::FScopedLockProperties ScopedLock(DestinationSkeletalMesh);
+	FSkinnedAssetAsyncBuildScope AsyncBuildScope(DestinationSkeletalMesh);
 
 	// Get a list of all the clothing assets affecting this LOD so we can re-apply later
 	TArray<ClothingAssetUtils::FClothingAssetMeshBinding> ClothingBindings;
@@ -527,6 +548,9 @@ bool FLODUtilities::SetCustomLOD(USkeletalMesh* DestinationSkeletalMesh, USkelet
 	TArray<int32> ClothingAssetSectionIndices;
 	TArray<int32> ClothingAssetInternalLodIndices;
 
+	TArray<FSkinWeightProfileInfo> ExistingSkinWeightProfileInfos;
+	TArray<FSkeletalMeshImportData> ExistingAlternateImportDataPerLOD;
+
 	FSkeletalMeshModel* const SourceImportedResource = SourceSkeletalMesh->GetImportedModel();
 	FSkeletalMeshModel* const DestImportedResource = DestinationSkeletalMesh->GetImportedModel();
 
@@ -535,10 +559,29 @@ bool FLODUtilities::SetCustomLOD(USkeletalMesh* DestinationSkeletalMesh, USkelet
 		return false;
 	}
 
-	if (SourceImportedResource->LODModels.IsValidIndex(LodIndex))
+	if (DestImportedResource->LODModels.IsValidIndex(LodIndex))
 	{
-		bMustReimportAlternateSkinWeightProfile = true;
 		FLODUtilities::UnbindClothingAndBackup(DestinationSkeletalMesh, ClothingBindings, LodIndex);
+
+		int32 ExistingLodCount = DestinationSkeletalMesh->GetLODNum();
+
+		//Extract all LOD Skin weight profiles data, we will re-apply them at the end
+		ExistingSkinWeightProfileInfos = DestinationSkeletalMesh->GetSkinWeightProfiles();
+		DestinationSkeletalMesh->GetSkinWeightProfiles().Reset();
+		for (int32 AllLodIndex = 0; AllLodIndex < ExistingLodCount; ++AllLodIndex)
+		{
+			FSkeletalMeshLODModel& BuildLODModel = DestinationSkeletalMesh->GetImportedModel()->LODModels[AllLodIndex];
+			BuildLODModel.SkinWeightProfiles.Reset();
+
+			//Store the LOD alternate skinning profile data
+			FSkeletalMeshImportData& SkeletalMeshImportData = ExistingAlternateImportDataPerLOD.AddDefaulted_GetRef();
+			if (DestinationSkeletalMesh->HasMeshDescription(AllLodIndex))
+			{
+				PRAGMA_DISABLE_DEPRECATION_WARNINGS
+				DestinationSkeletalMesh->LoadLODImportedData(AllLodIndex, SkeletalMeshImportData);
+				PRAGMA_ENABLE_DEPRECATION_WARNINGS
+			}
+		}
 	}
 
 	//Lambda to call to re-apply the clothing
@@ -551,6 +594,72 @@ bool FLODUtilities::SetCustomLOD(USkeletalMesh* DestinationSkeletalMesh, USkelet
 		}
 	};
 
+	auto ReapplyAlternateSkinning = [&DestinationSkeletalMesh, &ExistingSkinWeightProfileInfos, &ExistingAlternateImportDataPerLOD, &LodIndex]()
+	{
+		if (ExistingSkinWeightProfileInfos.Num() > 0)
+		{
+			TArray<FSkinWeightProfileInfo>& SkinProfiles = DestinationSkeletalMesh->GetSkinWeightProfiles();
+			SkinProfiles = ExistingSkinWeightProfileInfos;
+			for (const FSkinWeightProfileInfo& ProfileInfo : SkinProfiles)
+			{
+				const int32 LodCount = DestinationSkeletalMesh->GetLODNum();
+				for(int32 AllLodIndex = 0; AllLodIndex < LodCount; ++AllLodIndex)
+				{
+					if (!ExistingAlternateImportDataPerLOD.IsValidIndex(AllLodIndex))
+					{
+						continue;
+					}
+					if (!DestinationSkeletalMesh->HasMeshDescription(AllLodIndex))
+					{
+						continue;
+					}
+					const FSkeletalMeshLODInfo* LodInfo = DestinationSkeletalMesh->GetLODInfo(AllLodIndex);
+					if (!LodInfo)
+					{
+						continue;
+					}
+
+					const FSkeletalMeshImportData& ExistingImportDataSrc = ExistingAlternateImportDataPerLOD[AllLodIndex];
+
+					const FString ProfileNameStr = ProfileInfo.Name.ToString();
+
+					FSkeletalMeshImportData ImportDataDest;
+					PRAGMA_DISABLE_DEPRECATION_WARNINGS
+					DestinationSkeletalMesh->LoadLODImportedData(AllLodIndex, ImportDataDest);
+					PRAGMA_ENABLE_DEPRECATION_WARNINGS
+
+					int32 PointNumberDest = ImportDataDest.Points.Num();
+					int32 VertexNumberDest = ImportDataDest.Points.Num();
+
+					if (ExistingImportDataSrc.Points.Num() != PointNumberDest)
+					{
+						UE_LOG(LogLODUtilities, Error, TEXT("Alternate skinning mesh vertex number is different from the mesh LOD, we cannot apply the existing alternate skinning [%s] on the re-import skeletal mesh LOD %d [%s]"), *ProfileNameStr, LodIndex, *DestinationSkeletalMesh->GetName());
+						continue;
+					}
+
+					//Replace the data into the destination bulk data and save it
+					int32 ProfileIndex = 0;
+					if (ImportDataDest.AlternateInfluenceProfileNames.Find(ProfileNameStr, ProfileIndex))
+					{
+						ImportDataDest.AlternateInfluenceProfileNames.RemoveAt(ProfileIndex);
+						ImportDataDest.AlternateInfluences.RemoveAt(ProfileIndex);
+					}
+					int32 SrcProfileIndex = 0;
+					if (ExistingImportDataSrc.AlternateInfluenceProfileNames.Find(ProfileNameStr, SrcProfileIndex))
+					{
+						ImportDataDest.AlternateInfluenceProfileNames.Add(ProfileNameStr);
+						ImportDataDest.AlternateInfluences.Add(ExistingImportDataSrc.AlternateInfluences[SrcProfileIndex]);
+					}
+
+					//Resave the bulk data with the new or refreshed data
+					PRAGMA_DISABLE_DEPRECATION_WARNINGS
+					DestinationSkeletalMesh->SaveLODImportedData(AllLodIndex, ImportDataDest);
+					PRAGMA_ENABLE_DEPRECATION_WARNINGS
+				}
+			}
+		}
+	};
+
 	// Now we copy the base FSkeletalMeshLODModel from the imported skeletal mesh as the new LOD in the selected mesh.
 	if (SourceImportedResource->LODModels.Num() == 0)
 	{
@@ -601,8 +710,6 @@ bool FLODUtilities::SetCustomLOD(USkeletalMesh* DestinationSkeletalMesh, USkelet
 		}
 	}
 
-	FScopedSkeletalMeshPostEditChange ScopedPostEditChange(DestinationSkeletalMesh);
-
 	FSkeletalMeshLODModel& NewLODModel = SourceImportedResource->LODModels[0];
 
 	// If this LOD is not the base LOD, we check all bones we have sockets on are present in it.
@@ -842,12 +949,7 @@ bool FLODUtilities::SetCustomLOD(USkeletalMesh* DestinationSkeletalMesh, USkelet
 
 	ReapplyClothing();
 
-	//Must be the last step because it cleanup the fbx importer to import the alternate skinning FBX
-	if (bMustReimportAlternateSkinWeightProfile)
-	{
-		//TODO port skin weights utilities outside of UnrealEd module
-		//FSkinWeightsUtilities::ReimportAlternateSkinWeight(DestinationSkeletalMesh, LodIndex);
-	}
+	ReapplyAlternateSkinning();
 	
 	// Notification of success
 	FNotificationInfo NotificationInfo(FText::GetEmpty());
@@ -4158,6 +4260,7 @@ void FLODUtilities::ReorderMaterialSlotToBaseLod(USkeletalMesh* SkeletalMesh)
 			{
 				if (MaterialSlotRemap[MaterialIndex] != INDEX_NONE)
 				{
+					bFoundMatch = true;
 					continue;
 				}
 				int32& RemapIndex = MaterialSlotRemap[MaterialIndex];
diff --git a/Engine/Source/Editor/Persona/Private/AnimationEditorViewportClient.cpp b/Engine/Source/Editor/Persona/Private/AnimationEditorViewportClient.cpp
--- a/Engine/Source/Editor/Persona/Private/AnimationEditorViewportClient.cpp
+++ b/Engine/Source/Editor/Persona/Private/AnimationEditorViewportClient.cpp
@@ -65,6 +65,22 @@ namespace EAnimationPlaybackSpeeds
 	float Values[EAnimationPlaybackSpeeds::NumPlaybackSpeeds] = { 0.1f, 0.25f, 0.5f, 0.75f, 1.0f, 2.0f, 5.0f, 10.0f, 0.f };
 }
 
+namespace UE::Private
+{
+	bool CanDrawPreviewComponents(TArray<UDebugSkelMeshComponent*> PreviewComponents)
+	{
+		//Avoid drawing if any of the component reference a compiling asset
+		for (UDebugSkelMeshComponent* PreviewMeshComponent : PreviewComponents)
+		{
+			if (PreviewMeshComponent && PreviewMeshComponent->GetSkinnedAsset() && PreviewMeshComponent->GetSkinnedAsset()->IsCompiling())
+			{
+				return false;
+			}
+		}
+		return true;
+	}
+}
+
 #define LOCTEXT_NAMESPACE "FAnimationViewportClient"
 
 /////////////////////////////////////////////////////////////////////////
@@ -519,10 +535,15 @@ void FAnimationViewportClient::HandleOnSkelMeshPhysicsCreated()
 
 void FAnimationViewportClient::Draw(const FSceneView* View, FPrimitiveDrawInterface* PDI)
 {
+	TArray<UDebugSkelMeshComponent*> PreviewMeshComponents = GetPreviewScene()->GetAllPreviewMeshComponents();
+	if (!UE::Private::CanDrawPreviewComponents(PreviewMeshComponents))
+	{
+		return;
+	}
+
 	FEditorViewportClient::Draw(View, PDI);
 
 	// draw bones for all debug skeletal meshes
-	TArray<UDebugSkelMeshComponent*> PreviewMeshComponents = GetPreviewScene()->GetAllPreviewMeshComponents();
 	for (UDebugSkelMeshComponent* PreviewMeshComponent : PreviewMeshComponents)
 	{
 		const bool bValidComponent = PreviewMeshComponent != nullptr;
@@ -620,11 +641,16 @@ void FAnimationViewportClient::Draw(const FSceneView* View, FPrimitiveDrawInterf
 
 void FAnimationViewportClient::DrawCanvas( FViewport& InViewport, FSceneView& View, FCanvas& Canvas )
 {
+	TArray<UDebugSkelMeshComponent*> PreviewMeshComponents = GetPreviewScene()->GetAllPreviewMeshComponents();
+	if (!UE::Private::CanDrawPreviewComponents(PreviewMeshComponents))
+	{
+		return;
+	}
+
 	FEditorViewportClient::DrawCanvas(InViewport, View, Canvas);
 
 	UWorld* World = nullptr;
 	
-	TArray<UDebugSkelMeshComponent*> PreviewMeshComponents = GetPreviewScene()->GetAllPreviewMeshComponents();
 	for  (UDebugSkelMeshComponent* PreviewMeshComponent : PreviewMeshComponents)
 	{
 		if (!PreviewMeshComponent)
diff --git a/Engine/Source/Editor/Persona/Private/PersonaMeshDetails.cpp b/Engine/Source/Editor/Persona/Private/PersonaMeshDetails.cpp
--- a/Engine/Source/Editor/Persona/Private/PersonaMeshDetails.cpp
+++ b/Engine/Source/Editor/Persona/Private/PersonaMeshDetails.cpp
@@ -4653,23 +4653,25 @@ FReply FPersonaMeshDetails::OnReimportLodClicked(EReimportButtonType InReimportT
 			}
 		}
 
-		bool bImportSucceeded = FbxMeshUtils::ImportMeshLODDialog(SkelMesh, InLODIndex);
-
-		if(InReimportType == EReimportButtonType::ReimportWithNewFile && !bImportSucceeded)
-		{
-			// Copy old source file back, as this one failed
-			LODInfo->SourceImportFilename = SourceFilenameBackup;
-			if (bRestoreReductionOnfail)
+		FbxMeshUtils::ImportMeshLODDialog(SkelMesh, InLODIndex).Then([this, InReimportType, SkelMesh, InLODIndex, SourceFilenameBackup, bRestoreReductionOnfail, ReductionSettingsBackup](TFuture<bool> Result)
 			{
-				LODInfo->ReductionSettings = ReductionSettingsBackup;
-			}
-		}
-		else if(InReimportType == EReimportButtonType::ReimportWithNewFile)
-		{
-			//Refresh the layout so the BaseLOD min max get recompute
-			RefreshMeshDetailLayout();
-		}
-
+				bool bImportSucceeded = Result.Get();
+				if (InReimportType == EReimportButtonType::ReimportWithNewFile && !bImportSucceeded)
+				{
+					FSkeletalMeshLODInfo* LODInfo = SkelMesh->GetLODInfo(InLODIndex);
+					// Copy old source file back, as this one failed
+					LODInfo->SourceImportFilename = SourceFilenameBackup;
+					if (bRestoreReductionOnfail)
+					{
+						LODInfo->ReductionSettings = ReductionSettingsBackup;
+					}
+				}
+				else if (InReimportType == EReimportButtonType::ReimportWithNewFile)
+				{
+					//Refresh the layout so the BaseLOD min max get recompute
+					RefreshMeshDetailLayout();
+				}
+			});
 		return FReply::Handled();
 	}
 
diff --git a/Engine/Source/Editor/Persona/Private/SkeletonSelectionEditMode.cpp b/Engine/Source/Editor/Persona/Private/SkeletonSelectionEditMode.cpp
--- a/Engine/Source/Editor/Persona/Private/SkeletonSelectionEditMode.cpp
+++ b/Engine/Source/Editor/Persona/Private/SkeletonSelectionEditMode.cpp
@@ -489,7 +489,12 @@ bool FSkeletonSelectionEditMode::IsSelectedBoneRequired() const
 bool FSkeletonSelectionEditMode::ShouldDrawWidget() const
 {
 	UDebugSkelMeshComponent* PreviewMeshComponent = GetAnimPreviewScene().GetPreviewMeshComponent();
-	if (!PreviewMeshComponent->IsAnimBlueprintInstanced())
+	if (PreviewMeshComponent && PreviewMeshComponent->GetSkeletalMeshAsset() && PreviewMeshComponent->GetSkeletalMeshAsset()->IsCompiling())
+	{
+		return false;
+	}
+
+	if (PreviewMeshComponent && !PreviewMeshComponent->IsAnimBlueprintInstanced())
 	{
 		return IsSelectedBoneRequired() || GetAnimPreviewScene().GetSelectedSocket().IsValid() || GetAnimPreviewScene().GetSelectedActor() != nullptr;
 	}
diff --git a/Engine/Source/Editor/SkeletalMeshEditor/Private/SkeletalMeshEditor.cpp b/Engine/Source/Editor/SkeletalMeshEditor/Private/SkeletalMeshEditor.cpp
--- a/Engine/Source/Editor/SkeletalMeshEditor/Private/SkeletalMeshEditor.cpp
+++ b/Engine/Source/Editor/SkeletalMeshEditor/Private/SkeletalMeshEditor.cpp
@@ -18,6 +18,7 @@
 #include "EditorModeManager.h"
 #include "EditorReimportHandler.h"
 #include "EditorViewportClient.h"
+#include "Engine/SkinnedAssetAsyncCompileUtils.h"
 #include "EngineGlobals.h"
 #include "EngineUtils.h"
 #include "Factories/FbxSkeletalMeshImportData.h"
@@ -1267,166 +1268,193 @@ UObject* FSkeletalMeshEditor::HandleGetAsset()
 	return GetEditingObject();
 }
 
-bool FSkeletalMeshEditor::HandleReimportMeshInternal(int32 SourceFileIndex /*= INDEX_NONE*/, bool bWithNewFile /*= false*/)
+TFuture<bool> FSkeletalMeshEditor::HandleReimportMeshInternal(int32 SourceFileIndex /*= INDEX_NONE*/, bool bWithNewFile /*= false*/)
 {
-	bool bResult = false;
-	//Interchange reimport are asynchronous
-	if (UInterchangeAssetImportData* AssetImportData = Cast<UInterchangeAssetImportData>(SkeletalMesh->GetAssetImportData()))
-	{
-		UE::Interchange::FAssetImportResultRef Result = FReimportManager::Instance()->ReimportAsync(SkeletalMesh, true, true, TEXT(""), nullptr, SourceFileIndex, bWithNewFile);
-		
-		Result->OnDone([SkeletonTreePtr = SkeletonTree, WeakSkeletalMesh = TWeakObjectPtr<USkeletalMesh>(SkeletalMesh)](UE::Interchange::FImportResult& Result)
-			{
-				auto ResetComponent = [SkeletonTreePtr, WeakSkeletalMesh]()
+	TSharedPtr<TPromise<bool>> Promise = MakeShared<TPromise<bool>>();
+
+	//The reimport will be asynchronous only if the reimport manager use Interchange.
+	UE::Interchange::FAssetImportResultRef Result = FReimportManager::Instance()->ReimportAsync(SkeletalMesh, true, true, TEXT(""), nullptr, SourceFileIndex, bWithNewFile);
+
+	Result->OnDone([Promise, SkeletonTreePtr = SkeletonTree, WeakSkeletalMesh = TWeakObjectPtr<USkeletalMesh>(SkeletalMesh)](UE::Interchange::FImportResult& Result)
+		{
+			auto ResetComponent = [Promise, SkeletonTreePtr, WeakSkeletalMesh, &Result]()
 				{
-					FScopedSuspendAlternateSkinWeightPreview ScopedSuspendAlternateSkinnWeightPreview(WeakSkeletalMesh.Get());
+					// Refresh skeleton tree
+					SkeletonTreePtr->Refresh();
+
+					const TArray<UInterchangeResult*>& Results = Result.GetResults()->GetResults();
+					for (const UInterchangeResult* InterchangeResult : Results)
 					{
-						constexpr bool bCallPostEditChange = false;
-						constexpr bool bReregisterComponents = true;
-						FScopedSkeletalMeshPostEditChange ScopedPostEditChange = FScopedSkeletalMeshPostEditChange(WeakSkeletalMesh.Get(), bCallPostEditChange, bReregisterComponents);
-						// Refresh skeleton tree
-						SkeletonTreePtr->Refresh();
+						if (InterchangeResult->IsA<UInterchangeResultError_ReimportFail>())
+						{
+							Promise->SetValue(false);
+							return;
+						}
 					}
+					Promise->SetValue(true);
 				};
 
-				if (IsInGameThread())
-				{
-					ResetComponent();
-				}
-				else
-				{
-					Async(EAsyncExecution::TaskGraphMainThread, MoveTemp(ResetComponent));
-				}
-			});
-		
-		if (Result->GetStatus() == UE::Interchange::FImportResult::EStatus::Done)
-		{
-			const TArray<UInterchangeResult*>& Results = Result->GetResults()->GetResults();
-			for (const UInterchangeResult* InterchangeResult : Results)
+			if (IsInGameThread())
 			{
-				if (InterchangeResult->IsA<UInterchangeResultError_ReimportFail>())
-				{
-					return false;
-				}
+				ResetComponent();
 			}
-		}
-		return true;
-	}
-	else
-	{
-		FScopedSuspendAlternateSkinWeightPreview ScopedSuspendAlternateSkinnWeightPreview(SkeletalMesh);
-		{
-			FScopedSkeletalMeshPostEditChange ScopedPostEditChange(SkeletalMesh);
-			// Reimport the asset
-			bResult = FReimportManager::Instance()->Reimport(SkeletalMesh, true, true, TEXT(""), nullptr, SourceFileIndex, bWithNewFile);
-			// Refresh skeleton tree
-			SkeletonTree->Refresh();
-		}
-	}
-	return bResult;
+			else
+			{
+				Async(EAsyncExecution::TaskGraphMainThread, MoveTemp(ResetComponent));
+			}
+		});
+	return Promise->GetFuture();
 }
 
 void FSkeletalMeshEditor::HandleReimportMesh(int32 SourceFileIndex /*= INDEX_NONE*/)
 {
-	FScopedSuspendAlternateSkinWeightPreview ScopedSuspendAlternateSkinnWeightPreview(SkeletalMesh);
-
-	HandleReimportMeshInternal(SourceFileIndex, false);
+	TSharedPtr<FScopedSuspendAlternateSkinWeightPreview> ScopedSuspendAlternateSkinnWeightPreview = MakeShared<FScopedSuspendAlternateSkinWeightPreview>(SkeletalMesh);
+	TSharedPtr<FScopedSkeletalMeshReregisterContexts> ScopedReregisterComponents = MakeShared<FScopedSkeletalMeshReregisterContexts>(SkeletalMesh);
+	HandleReimportMeshInternal(SourceFileIndex, false).Then([ScopedSuspendAlternateSkinnWeightPreview, ScopedReregisterComponents](TFuture<bool> ReimportResult) {});
 }
 
 void FSkeletalMeshEditor::HandleReimportMeshWithNewFile(int32 SourceFileIndex /*= INDEX_NONE*/)
 {
-	FScopedSuspendAlternateSkinWeightPreview ScopedSuspendAlternateSkinnWeightPreview(SkeletalMesh);
-
-	HandleReimportMeshInternal(SourceFileIndex, true);
+	TSharedPtr<FScopedSuspendAlternateSkinWeightPreview> ScopedSuspendAlternateSkinnWeightPreview = MakeShared<FScopedSuspendAlternateSkinWeightPreview>(SkeletalMesh);
+	TSharedPtr<FScopedSkeletalMeshReregisterContexts> ScopedReregisterComponents = MakeShared<FScopedSkeletalMeshReregisterContexts>(SkeletalMesh);
+	HandleReimportMeshInternal(SourceFileIndex, true).Then([ScopedSuspendAlternateSkinnWeightPreview, ScopedReregisterComponents](TFuture<bool> ReimportResult) {});
 }
 
-void ReimportAllCustomLODs(USkeletalMesh* SkeletalMesh, UDebugSkelMeshComponent* PreviewMeshComponent, bool bWithNewFile)
+TFuture<bool> ReimportLodInChain(USkeletalMesh* SkeletalMesh
+	, UDebugSkelMeshComponent* PreviewMeshComponent
+	, bool bWithNewFile
+	, TSharedPtr<TArray<bool>> Dependencies
+	, int32 LodIndex)
 {
-	FScopedSuspendAlternateSkinWeightPreview ScopedSuspendAlternateSkinnWeightPreview(SkeletalMesh);
+	TSharedPtr<TPromise<bool>> Promise = MakeShared<TPromise<bool>>();
+
+	if (SkeletalMesh->GetLODNum() <= LodIndex)
+	{
+		//Nothing to do
+		Promise->SetValue(false);
+		return Promise->GetFuture();
+	}
+	
+	TArray<bool>& DependenciesRef = *Dependencies.Get();
+
+	//Do not reimport LOD that was re-import with the base mesh
+	if (!SkeletalMesh->GetLODInfo(LodIndex)->bImportWithBaseMesh)
 	{
-		FScopedSkeletalMeshPostEditChange ScopedPostEditChange(SkeletalMesh);
-
-		//Find the dependencies of the generated LOD
-		TArray<bool> Dependencies;
-		Dependencies.AddZeroed(SkeletalMesh->GetLODNum());
-		//Avoid making LOD 0 to true in the dependencies since everything that should be regenerate base on LOD 0 is already regenerate at this point.
-		//But we need to regenerate every generated LOD base on any re-import custom LOD
-		//Reimport all custom LODs
-		for (int32 LodIndex = 1; LodIndex < SkeletalMesh->GetLODNum(); ++LodIndex)
+		if (SkeletalMesh->GetLODInfo(LodIndex)->bHasBeenSimplified == false)
 		{
-			//Do not reimport LOD that was re-import with the base mesh
-			if (SkeletalMesh->GetLODInfo(LodIndex)->bImportWithBaseMesh)
+			FString SourceFilenameBackup = SkeletalMesh->GetLODInfo(LodIndex)->SourceImportFilename;
+			if (bWithNewFile)
 			{
-				continue;
+				SkeletalMesh->GetLODInfo(LodIndex)->SourceImportFilename.Empty();
 			}
-			if (SkeletalMesh->GetLODInfo(LodIndex)->bHasBeenSimplified == false)
-			{
-				FString SourceFilenameBackup = SkeletalMesh->GetLODInfo(LodIndex)->SourceImportFilename;
-				if (bWithNewFile)
-				{
-					SkeletalMesh->GetLODInfo(LodIndex)->SourceImportFilename.Empty();
-				}
 
-				if (!FbxMeshUtils::ImportMeshLODDialog(SkeletalMesh, LodIndex, false))
+			FbxMeshUtils::ImportMeshLODDialog(SkeletalMesh, LodIndex, false).Then([Promise, SkeletalMesh, bWithNewFile, LodIndex, SourceFilenameBackup, PreviewMeshComponent, Dependencies](TFuture<bool> FutureResult)
 				{
-					if (bWithNewFile)
+					TArray<bool>& DependenciesRef = *Dependencies.Get();
+					const bool bResult = FutureResult.Get();
+					if (!bResult)
 					{
-						SkeletalMesh->GetLODInfo(LodIndex)->SourceImportFilename = SourceFilenameBackup;
+						if (bWithNewFile)
+						{
+							SkeletalMesh->GetLODInfo(LodIndex)->SourceImportFilename = SourceFilenameBackup;
+						}
 					}
-				}
-				else
-				{
-					Dependencies[LodIndex] = true;
-				}
-			}
-			else if (Dependencies[SkeletalMesh->GetLODInfo(LodIndex)->ReductionSettings.BaseLOD])
-			{
-				//Regenerate the LOD
-				FSkeletalMeshUpdateContext UpdateContext;
-				UpdateContext.SkeletalMesh = SkeletalMesh;
-				UpdateContext.AssociatedComponents.Push(PreviewMeshComponent);
-				FLODUtilities::SimplifySkeletalMeshLOD(UpdateContext, LodIndex, GetTargetPlatformManagerRef().GetRunningTargetPlatform(), false);
-				Dependencies[LodIndex] = true;
-			}
+					else
+					{
+						DependenciesRef[LodIndex] = true;
+					}
+
+					//Iterate the next lod chain
+					if (SkeletalMesh->GetLODNum() > LodIndex + 1)
+					{
+						ReimportLodInChain(SkeletalMesh, PreviewMeshComponent, bWithNewFile, Dependencies, LodIndex + 1).Then([Promise](TFuture<bool> ReimportLodInChainFutureResult)
+							{
+								const bool bReimportLodInChainResult = ReimportLodInChainFutureResult.Get();
+								Promise->SetValue(bReimportLodInChainResult);
+							});
+					}
+					else
+					{
+						Promise->SetValue(bResult);
+					}
+				});
+			return Promise->GetFuture();
 		}
+		else if (DependenciesRef[SkeletalMesh->GetLODInfo(LodIndex)->ReductionSettings.BaseLOD])
+		{
+			//Regenerate the LOD
+			FSkeletalMeshUpdateContext UpdateContext;
+			UpdateContext.SkeletalMesh = SkeletalMesh;
+			UpdateContext.AssociatedComponents.Push(PreviewMeshComponent);
+			FLODUtilities::SimplifySkeletalMeshLOD(UpdateContext, LodIndex, GetTargetPlatformManagerRef().GetRunningTargetPlatform(), false);
+			DependenciesRef[LodIndex] = true;
+		}
+	}
+	//Iterate the next Lod in the chain
+	if (SkeletalMesh->GetLODNum() > LodIndex + 1)
+	{
+		ReimportLodInChain(SkeletalMesh, PreviewMeshComponent, bWithNewFile, Dependencies, LodIndex + 1).Then([Promise](TFuture<bool> ReimportLodInChainFutureResult)
+			{
+				const bool bReimportLodInChainResult = ReimportLodInChainFutureResult.Get();
+				Promise->SetValue(bReimportLodInChainResult);
+			});
+	}
+	else
+	{
+		//We have iterate all LodIndex set the promise
+		Promise->SetValue(true);
 	}
+	return Promise->GetFuture();
 }
 
-void FSkeletalMeshEditor::HandleReimportAllMesh(int32 SourceFileIndex /*= INDEX_NONE*/)
+TFuture<bool> ReimportAllCustomLODs(USkeletalMesh* SkeletalMesh, UDebugSkelMeshComponent* PreviewMeshComponent, bool bWithNewFile)
 {
-	FScopedSuspendAlternateSkinWeightPreview ScopedSuspendAlternateSkinnWeightPreview(SkeletalMesh);
-	// Reimport the asset
-	if (SkeletalMesh)
-	{
-		FScopedSkeletalMeshPostEditChange ScopedPostEditChange(SkeletalMesh);
+	TSharedPtr<TPromise<bool>> Promise = MakeShared<TPromise<bool>>();
+	//Find the dependencies of the generated LOD
+	TSharedPtr<TArray<bool>> Dependencies = MakeShared<TArray<bool>>();
+	Dependencies->AddZeroed(SkeletalMesh->GetLODNum());
 
-		//Reimport base LOD
-		if (HandleReimportMeshInternal(SourceFileIndex, false))
+	int32 LodIndex = 1;
+	ReimportLodInChain(SkeletalMesh, PreviewMeshComponent, bWithNewFile, Dependencies, LodIndex).Then([Promise](TFuture<bool> FutureResult)
 		{
-			//Reimport all custom LODs
-			ReimportAllCustomLODs(SkeletalMesh, GetPersonaToolkit()->GetPreviewMeshComponent(), false);
-		}
-	}
+			bool bResult = FutureResult.Get();
+			Promise->SetValue(bResult);
+		});
+	return Promise->GetFuture();
 }
 
-void FSkeletalMeshEditor::HandleReimportAllMeshWithNewFile(int32 SourceFileIndex /*= INDEX_NONE*/)
+void FSkeletalMeshEditor::HandleReimportAllMeshInternal(int32 SourceFileIndex, bool bWithNewFile)
 {
-	FScopedSuspendAlternateSkinWeightPreview ScopedSuspendAlternateSkinnWeightPreview(SkeletalMesh);
-
 	// Reimport the asset
 	if (SkeletalMesh)
 	{
-		FScopedSkeletalMeshPostEditChange ScopedPostEditChange(SkeletalMesh);
-		TArray<UObject*> ImportObjs;
-		ImportObjs.Add(SkeletalMesh);
-		if (HandleReimportMeshInternal(SourceFileIndex, true))
+		TSharedPtr<FScopedSuspendAlternateSkinWeightPreview> ScopedSuspendAlternateSkinnWeightPreview = MakeShared<FScopedSuspendAlternateSkinWeightPreview>(SkeletalMesh);
+		TSharedPtr<FScopedSkeletalMeshReregisterContexts> ScopedReregisterComponents = MakeShared<FScopedSkeletalMeshReregisterContexts>(SkeletalMesh);
+		//Reimport base LOD
+		HandleReimportMeshInternal(SourceFileIndex, bWithNewFile).Then([this, bWithNewFile, ScopedSuspendAlternateSkinnWeightPreview, ScopedReregisterComponents](TFuture<bool> Result)
 		{
-			//Reimport all custom LODs
-			ReimportAllCustomLODs(SkeletalMesh, GetPersonaToolkit()->GetPreviewMeshComponent(), true);
-		}
+			check(IsInGameThread());
+			//import all custom LODs
+			if (Result.Get() && SkeletalMesh->GetLODNum() > 1)
+			{
+				ReimportAllCustomLODs(SkeletalMesh.Get(), GetPersonaToolkit()->GetPreviewMeshComponent(), bWithNewFile);
+			}
+		});
 	}
 }
 
+void FSkeletalMeshEditor::HandleReimportAllMesh(int32 SourceFileIndex /*= INDEX_NONE*/)
+{
+	constexpr bool bWithNewFile = false;
+	HandleReimportAllMeshInternal(SourceFileIndex, bWithNewFile);
+}
+
+void FSkeletalMeshEditor::HandleReimportAllMeshWithNewFile(int32 SourceFileIndex /*= INDEX_NONE*/)
+{
+	constexpr bool bWithNewFile = true;
+	HandleReimportAllMeshInternal(SourceFileIndex, bWithNewFile);
+}
+
 void FSkeletalMeshEditor::HandleOnPreviewSceneSettingsCustomized(IDetailLayoutBuilder& DetailBuilder)
 {
 	DetailBuilder.HideCategory("Mesh");
diff --git a/Engine/Source/Editor/SkeletalMeshEditor/Private/SkeletalMeshEditor.h b/Engine/Source/Editor/SkeletalMeshEditor/Private/SkeletalMeshEditor.h
--- a/Engine/Source/Editor/SkeletalMeshEditor/Private/SkeletalMeshEditor.h
+++ b/Engine/Source/Editor/SkeletalMeshEditor/Private/SkeletalMeshEditor.h
@@ -11,6 +11,7 @@
 #include "ISkeletalMeshEditor.h"
 #include "SkeletalMeshNotifier.h"
 #include "Containers/ArrayView.h"
+#include "Async/Future.h"
 
 class IDetailLayoutBuilder;
 class IDetailsView;
@@ -128,10 +129,11 @@ private:
 
 	void HandleReimportMesh(int32 SourceFileIndex = INDEX_NONE);
 	void HandleReimportMeshWithNewFile(int32 SourceFileIndex = INDEX_NONE);
-	
-	bool HandleReimportMeshInternal(int32 SourceFileIndex = INDEX_NONE, bool bWithNewFile = false);
+	TFuture<bool> HandleReimportMeshInternal(int32 SourceFileIndex = INDEX_NONE, bool bWithNewFile = false);
+
 	void HandleReimportAllMesh(int32 SourceFileIndex = INDEX_NONE);
 	void HandleReimportAllMeshWithNewFile(int32 SourceFileIndex = INDEX_NONE);
+	void HandleReimportAllMeshInternal(int32 SourceFileIndex, bool bWithNewFile);
 
 	void HandleOnPreviewSceneSettingsCustomized(IDetailLayoutBuilder& DetailBuilder);
 
diff --git a/Engine/Source/Editor/StaticMeshEditor/Private/StaticMeshEditorTools.cpp b/Engine/Source/Editor/StaticMeshEditor/Private/StaticMeshEditorTools.cpp
--- a/Engine/Source/Editor/StaticMeshEditor/Private/StaticMeshEditorTools.cpp
+++ b/Engine/Source/Editor/StaticMeshEditor/Private/StaticMeshEditorTools.cpp
@@ -4358,16 +4358,18 @@ void FLevelOfDetailSettingsLayout::OnImportLOD(TSharedPtr<FString> NewValue, ESe
 		//This boolean prevent changing the value when the LOD is reimport
 		bool bImportCustomLOD = (LODIndex >= StaticMesh->GetNumSourceModels());
 
-		bool bResult = FbxMeshUtils::ImportMeshLODDialog(StaticMesh, LODIndex);
+		FbxMeshUtils::ImportMeshLODDialog(StaticMesh, LODIndex).Then([this, StaticMesh, bImportCustomLOD, LODIndex](TFuture<bool> FutureResult)
+			{
+				bool bResult = FutureResult.Get();
+				if (bImportCustomLOD && bResult && StaticMesh->IsSourceModelValid(LODIndex))
+				{
+					//Custom LOD should reduce base on them self when they get imported.
+					StaticMesh->GetSourceModel(LODIndex).ReductionSettings.BaseLODModel = LODIndex;
+				}
 
-		if (bImportCustomLOD && bResult && StaticMesh->IsSourceModelValid(LODIndex))
-		{
-			//Custom LOD should reduce base on them self when they get imported.
-			StaticMesh->GetSourceModel(LODIndex).ReductionSettings.BaseLODModel = LODIndex;
-		}
-		
-		StaticMesh->PostEditChange();
-		StaticMeshEditor.RefreshTool();
+				StaticMesh->PostEditChange();
+				StaticMeshEditor.RefreshTool();
+			});
 	}
 
 }
diff --git a/Engine/Source/Editor/UnrealEd/Private/FbxMeshUtils.cpp b/Engine/Source/Editor/UnrealEd/Private/FbxMeshUtils.cpp
--- a/Engine/Source/Editor/UnrealEd/Private/FbxMeshUtils.cpp
+++ b/Engine/Source/Editor/UnrealEd/Private/FbxMeshUtils.cpp
@@ -190,7 +190,7 @@ namespace FbxMeshUtils
 		{
 			UInterchangeSourceData* SourceData = UInterchangeManager::GetInterchangeManager().CreateSourceData(Filename);
 			//Call interchange mesh utilities to import custom LOD
-			UInterchangeMeshUtilities::ImportCustomLodAsync(BaseStaticMesh, LODLevel, SourceData).Then([BaseStaticMesh, LODLevel](TFuture<bool> Result)
+			UInterchangeMeshUtilities::ImportCustomLod(BaseStaticMesh, LODLevel, SourceData).Then([BaseStaticMesh, LODLevel](TFuture<bool> Result)
 				{
 					bool bResult = Result.Get();
 					Async(EAsyncExecution::TaskGraphMainThread, [BaseStaticMesh, LODLevel, bResult]()
@@ -369,7 +369,7 @@ namespace FbxMeshUtils
 
 			UInterchangeSourceData* SourceData = UInterchangeManager::GetInterchangeManager().CreateSourceData(Filename);
 			//Call interchange mesh utilities to import custom LOD
-			UInterchangeMeshUtilities::ImportCustomLodAsync(TempStaticMesh, 0, SourceData).Then([BaseStaticMesh,TempStaticMesh](TFuture<bool> Result)
+			UInterchangeMeshUtilities::ImportCustomLod(TempStaticMesh, 0, SourceData).Then([BaseStaticMesh,TempStaticMesh](TFuture<bool> Result)
 			{
 				bool bResult = Result.Get();
 				Async(EAsyncExecution::TaskGraphMainThread, [BaseStaticMesh, TempStaticMesh, bResult]()
@@ -500,7 +500,7 @@ namespace FbxMeshUtils
 		{
 			UInterchangeSourceData* SourceData = UInterchangeManager::GetInterchangeManager().CreateSourceData(Filename);
 			//Call interchange mesh utilities to import custom LOD
-			UInterchangeMeshUtilities::ImportCustomLodAsync(SelectedSkelMesh, LODLevel, SourceData).Then([SelectedSkelMesh, LODLevel](TFuture<bool> Result)
+			UInterchangeMeshUtilities::ImportCustomLod(SelectedSkelMesh, LODLevel, SourceData).Then([SelectedSkelMesh, LODLevel](TFuture<bool> Result)
 				{
 					bool bResult = Result.Get();
 					Async(EAsyncExecution::TaskGraphMainThread, [SelectedSkelMesh, LODLevel, bResult]()
@@ -888,11 +888,13 @@ namespace FbxMeshUtils
 		return ChosenFilname;
 	}
 
-	bool ImportMeshLODDialog(class UObject* SelectedMesh, int32 LODLevel, bool bNotifyCB /*= true*/, bool bReimportWithNewFile /*= false*/)
+	TFuture<bool> ImportMeshLODDialog(class UObject* SelectedMesh, int32 LODLevel, bool bNotifyCB /*= true*/, bool bReimportWithNewFile /*= false*/)
 	{
+		TSharedPtr<TPromise<bool>> Promise = MakeShared<TPromise<bool>>();
 		if(!SelectedMesh)
 		{
-			return false;
+			Promise->SetValue(false);
+			return Promise->GetFuture();
 		}
 
 		USkeletalMesh* SkeletalMesh = Cast<USkeletalMesh>(SelectedMesh);
@@ -941,7 +943,8 @@ namespace FbxMeshUtils
 		else
 		{
 			//We support only staticmesh and skeletalmesh asset for LOD import
-			return false;
+			Promise->SetValue(false);
+			return Promise->GetFuture();
 		}
 
 
@@ -949,7 +952,8 @@ namespace FbxMeshUtils
 		{
 			UE_LOG(LogExportMeshUtils, Warning, TEXT("ImportMeshLODDialog: Invalid mesh LOD index %d, no prior LOD index exists."), LODLevel);
 			FbxMeshUtils::Private::ShowFailedToImportLodDialog(LODLevel);
-			return false;
+			Promise->SetValue(false);
+			return Promise->GetFuture();
 		}
 
 		// Check the file exists first
@@ -960,12 +964,24 @@ namespace FbxMeshUtils
 		//We will use interchange only if interchange is enable and the skeletalmesh we want to add a LOD was import with interchange
 		if(UInterchangeManager::IsInterchangeImportEnabled() && SelectedInterchangeAssetImportData)
 		{
-			auto CustomLodImportContinuation = [bNotifyCB, SkeletalMesh, StaticMesh, LODLevel](TFuture<bool> Result)
+			TFuture<bool> Result;
+			if(!bSourceFileExists)
 			{
-				bool bResult = Result.Get();
-				Async(EAsyncExecution::TaskGraphMainThread, [bNotifyCB, SkeletalMesh, StaticMesh, LODLevel, bResult]()
+				//Call interchange mesh utilities to import custom LOD
+				Result = UInterchangeMeshUtilities::ImportCustomLod(SelectedMesh, LODLevel);
+			}
+			else
+			{
+				UInterchangeSourceData* SourceData = UInterchangeManager::GetInterchangeManager().CreateSourceData(FilenameToImport);
+				Result = UInterchangeMeshUtilities::ImportCustomLod(SelectedMesh, LODLevel, SourceData);
+			}
+			Result.Then([Promise, bNotifyCB, SkeletalMesh, StaticMesh, LODLevel](TFuture<bool> FutureResult)
+				{
+					check(IsInGameThread());
+					bool bResult = FutureResult.Get();
+					if (bResult)
 					{
-						if (bNotifyCB && bResult)
+						if (bNotifyCB)
 						{
 							if (SkeletalMesh)
 							{
@@ -976,24 +992,16 @@ namespace FbxMeshUtils
 								GEditor->GetEditorSubsystem<UImportSubsystem>()->BroadcastAssetPostLODImport(StaticMesh, LODLevel);
 							}
 						}
-						else
-						{
-							FbxMeshUtils::Private::ShowFailedToImportLodDialog(LODLevel);
-						}
-					});
-			};
+					}
+					else
+					{
+						FbxMeshUtils::Private::ShowFailedToImportLodDialog(LODLevel);
+					}
+					Promise->SetValue(bResult);
+				});
+			
 
-			if(!bSourceFileExists)
-			{
-				//Call interchange mesh utilities to import custom LOD
-				UInterchangeMeshUtilities::ImportCustomLodAsync(SelectedMesh, LODLevel).Then(CustomLodImportContinuation);
-			}
-			else
-			{
-				UInterchangeSourceData* SourceData = UInterchangeManager::GetInterchangeManager().CreateSourceData(FilenameToImport);
-				UInterchangeMeshUtilities::ImportCustomLodAsync(SelectedMesh, LODLevel, SourceData).Then(CustomLodImportContinuation);
-			}
-			return true;
+			return Promise->GetFuture();
 		}
 		
 		if(!bSourceFileExists || FilenameToImport.IsEmpty())
@@ -1065,7 +1073,8 @@ namespace FbxMeshUtils
 			}
 		}
 
-		return bImportSuccess;
+		Promise->SetValue(bImportSuccess);
+		return Promise->GetFuture();
 	}
 
 	bool ImportStaticMeshHiResSourceModelDialog( UStaticMesh* StaticMesh )
diff --git a/Engine/Source/Editor/UnrealEd/Private/SkinWeightsUtilities.cpp b/Engine/Source/Editor/UnrealEd/Private/SkinWeightsUtilities.cpp
--- a/Engine/Source/Editor/UnrealEd/Private/SkinWeightsUtilities.cpp
+++ b/Engine/Source/Editor/UnrealEd/Private/SkinWeightsUtilities.cpp
@@ -89,12 +89,13 @@ bool FSkinWeightsUtilities::ImportAlternateSkinWeight(USkeletalMesh* SkeletalMes
 				//Avoid temporary package to be saved
 				UPackage* Package = ObjToDelete->GetOutermost();
 				Package->SetDirtyFlag(false);
-				//Avoid temporary asset to be saved by setting the RF_Transient flag
+				//Avoid gc, use keep flags
+				ObjToDelete->ClearFlags(RF_Standalone);
+				ObjToDelete->ClearInternalFlags(EInternalObjectFlags::Async);
+				//Make the object transient to prevent saving
 				ObjToDelete->SetFlags(RF_Transient);
 			}
 		}
-		ObjectTools::DeleteAssets(AssetsToDelete, false);
-		CollectGarbage(GARBAGE_COLLECTION_KEEPFLAGS);
 	};
 
 	DeletePathAssets();
@@ -102,6 +103,8 @@ bool FSkinWeightsUtilities::ImportAlternateSkinWeight(USkeletalMesh* SkeletalMes
 	UObject* ImportedObject = nullptr;
 	UnFbx::FBXImportOptions ImportOptions;
 
+	bool bCreateTransaction = false;
+
 	//Only use interchange if the base skeletal mesh was imported with interchange
 	const UInterchangeAssetImportData* SelectedInterchangeAssetImportData = Cast<UInterchangeAssetImportData>(SkeletalMesh->GetAssetImportData());
 	if (bUseInterchangeFramework && SelectedInterchangeAssetImportData)
@@ -208,6 +211,7 @@ bool FSkinWeightsUtilities::ImportAlternateSkinWeight(USkeletalMesh* SkeletalMes
 		//Factory and task can now be garbage collected
 		Task->RemoveFromRoot();
 		FbxFactory->RemoveFromRoot();
+		bCreateTransaction = true;
 	}
 
 	USkeletalMesh* TmpSkeletalMesh = Cast<USkeletalMesh>(ImportedObject);
@@ -235,10 +239,13 @@ bool FSkinWeightsUtilities::ImportAlternateSkinWeight(USkeletalMesh* SkeletalMes
 				FSkinWeightProfileInfo* Profile = SkeletalMesh->GetSkinWeightProfiles().FindByPredicate([ProfileName](FSkinWeightProfileInfo Profile) { return Profile.Name == ProfileName; });
 
 				const bool bIsReimportLocal = Profile != nullptr;
-				FText TransactionName = bIsReimportLocal ? NSLOCTEXT("UnrealEd", "UpdateAlternateSkinningWeight", "Update Alternate Skinning Weight")
-					: NSLOCTEXT("UnrealEd", "ImportAlternateSkinningWeight", "Import Alternate Skinning Weight");
-				FScopedTransaction ScopedTransaction(TransactionName);
-				SkeletalMesh->Modify();
+				if (bCreateTransaction)
+				{
+					FText TransactionName = bIsReimportLocal ? NSLOCTEXT("UnrealEd", "UpdateAlternateSkinningWeight", "Update Alternate Skinning Weight")
+						: NSLOCTEXT("UnrealEd", "ImportAlternateSkinningWeight", "Import Alternate Skinning Weight");
+					FScopedTransaction ScopedTransaction(TransactionName);
+					SkeletalMesh->Modify();
+				}
 
 				if (bIsReimportLocal)
 				{
diff --git a/Engine/Source/Editor/UnrealEd/Private/Toolkits/AssetEditorToolkit.cpp b/Engine/Source/Editor/UnrealEd/Private/Toolkits/AssetEditorToolkit.cpp
--- a/Engine/Source/Editor/UnrealEd/Private/Toolkits/AssetEditorToolkit.cpp
+++ b/Engine/Source/Editor/UnrealEd/Private/Toolkits/AssetEditorToolkit.cpp
@@ -223,7 +223,7 @@ void FAssetEditorToolkit::InitAssetEditor( const EToolkitMode::Type Mode, const
 		NewMajorTab->SetContent
 		( 
 			SAssignNew( NewStandaloneHost, SStandaloneAssetEditorToolkitHost, NewTabManager, AppIdentifier )
-			.IsEnabled_Lambda([ObjectsToEditWeak]()
+			.Visibility_Lambda([ObjectsToEditWeak]()
 				{
 					for (const TWeakObjectPtr<UObject> Object : ObjectsToEditWeak)
 					{
@@ -231,11 +231,11 @@ void FAssetEditorToolkit::InitAssetEditor( const EToolkitMode::Type Mode, const
 						{
 							if (AsyncAsset->IsCompiling())
 							{
-								return false;
+								return EVisibility::Collapsed;
 							}
 						}
 					}
-					return true;
+					return EVisibility::All;
 				})
 			.OnRequestClose(this, &FAssetEditorToolkit::OnRequestClose, EAssetEditorCloseReason::AssetEditorHostClosed)
 			.OnClose(this, &FAssetEditorToolkit::OnClose)
diff --git a/Engine/Source/Editor/UnrealEd/Public/FbxMeshUtils.h b/Engine/Source/Editor/UnrealEd/Public/FbxMeshUtils.h
--- a/Engine/Source/Editor/UnrealEd/Public/FbxMeshUtils.h
+++ b/Engine/Source/Editor/UnrealEd/Public/FbxMeshUtils.h
@@ -2,6 +2,7 @@
 
 #pragma once
 
+#include "Async/Future.h"
 #include "CoreMinimal.h"
 
 class UFbxImportUI;
@@ -48,7 +49,7 @@ namespace FbxMeshUtils
 	 * @param LODLevel			The level of the lod to import
 	 * @return					Whether or not the import succeeded
 	 */
-	UNREALED_API bool ImportMeshLODDialog( UObject* Mesh, int32 LODLevel, bool bNotifyCB = true, bool bReimportWithNewFile = false);
+	UNREALED_API TFuture<bool> ImportMeshLODDialog( UObject* Mesh, int32 LODLevel, bool bNotifyCB = true, bool bReimportWithNewFile = false);
 
 	/**
 	 * Imports a mesh as the high res source model of the given mesh
diff --git a/Engine/Source/Runtime/Engine/Private/SkeletalMesh.cpp b/Engine/Source/Runtime/Engine/Private/SkeletalMesh.cpp
--- a/Engine/Source/Runtime/Engine/Private/SkeletalMesh.cpp
+++ b/Engine/Source/Runtime/Engine/Private/SkeletalMesh.cpp
@@ -343,6 +343,52 @@ void FScopedSkeletalMeshPostEditChange::SetSkeletalMesh(USkeletalMesh* InSkeleta
 	}
 }
 
+FScopedSkeletalMeshReregisterContexts::FScopedSkeletalMeshReregisterContexts(USkeletalMesh* InSkeletalMesh)
+{
+	check(IsInGameThread());
+	SkeletalMesh = InSkeletalMesh;
+	if (!ensure(SkeletalMesh))
+	{
+		return;
+	}
+
+	//Make sure all components using this skeletalmesh have there render ressources free
+	RecreateExistingRenderStateContext = new FSkinnedMeshComponentRecreateRenderStateContext(SkeletalMesh, false);
+
+	// Now iterate over all skeletal mesh components and unregister them from the world, we will reregister them in the destructor
+	for (TObjectIterator<USkeletalMeshComponent> It; It; ++It)
+	{
+		USkeletalMeshComponent* SkelComp = *It;
+		if (SkelComp->GetSkeletalMeshAsset() == SkeletalMesh)
+		{
+			ComponentReregisterContexts.Add(new FComponentReregisterContext(SkelComp));
+		}
+	}
+}
+
+FScopedSkeletalMeshReregisterContexts::~FScopedSkeletalMeshReregisterContexts()
+{
+	check(IsInGameThread());
+	if (!ensure(SkeletalMesh))
+	{
+		return;
+	}
+
+	if (SkeletalMesh->IsCompiling())
+	{
+		//wait until the compilation is done before reregister the component
+		FSkinnedAssetCompilingManager::Get().FinishCompilation({ SkeletalMesh });
+	}
+
+	//Recreate the render context by calling delete
+	if (RecreateExistingRenderStateContext)
+	{
+		delete RecreateExistingRenderStateContext;
+	}
+
+	//Component will be reregister when going out of scope
+}
+
 const FString& GetSkeletalMeshDerivedDataVersion()
 {
 	static FString CachedVersionString = FDevSystemGuids::GetSystemGuid(FDevSystemGuids::Get().SkeletalMeshDerivedDataVersion).ToString();
diff --git a/Engine/Source/Runtime/Engine/Public/SkeletalMeshTypes.h b/Engine/Source/Runtime/Engine/Public/SkeletalMeshTypes.h
--- a/Engine/Source/Runtime/Engine/Public/SkeletalMeshTypes.h
+++ b/Engine/Source/Runtime/Engine/Public/SkeletalMeshTypes.h
@@ -185,4 +185,26 @@ private:
 	TIndirectArray<FComponentReregisterContext> ComponentReregisterContexts;
 };
 
+//Helper to scope the component register context.
+class FScopedSkeletalMeshReregisterContexts
+{
+public:
+	/*
+	 * This constructor will enregister all skeletal mesh component from the world. The component will also release there rendering resources.
+	 * Parameters:
+	 */
+	ENGINE_API FScopedSkeletalMeshReregisterContexts(USkeletalMesh* InSkeletalMesh);
+
+	/*
+	 * This destructor will Reregister all unregistered component to the world and there render data resources will be rebuild.
+	 */
+	ENGINE_API ~FScopedSkeletalMeshReregisterContexts();
+
+
+private:
+	USkeletalMesh* SkeletalMesh;
+	FSkinnedMeshComponentRecreateRenderStateContext* RecreateExistingRenderStateContext = nullptr;
+	TIndirectArray<FComponentReregisterContext> ComponentReregisterContexts;
+};
+
 #endif
diff --git a/Engine/Source/Runtime/Interchange/Core/Public/InterchangeFactoryBase.h b/Engine/Source/Runtime/Interchange/Core/Public/InterchangeFactoryBase.h
--- a/Engine/Source/Runtime/Interchange/Core/Public/InterchangeFactoryBase.h
+++ b/Engine/Source/Runtime/Interchange/Core/Public/InterchangeFactoryBase.h
@@ -29,6 +29,17 @@ enum class EInterchangeFactoryAssetType : uint8
 	Physics
 };
 
+/**
+ * Base class for post import task.
+ * Post import tasks are execute by the InterchangeManager when there is no import task to execute.
+ * See more detail in the Interchange manager header.
+ */
+class INTERCHANGECORE_API FInterchangePostImportTask
+{
+public:
+	virtual void Execute() {};
+};
+
 /**
  * Asset factory implementation:
  * 
@@ -41,7 +52,7 @@ enum class EInterchangeFactoryAssetType : uint8
  * The last two steps are helpful to change the imported/re-imported UObject before and after the PostEditChange (render data build) is call on the asset.
  * 
  * 4. SetupObject_GameThread - Do any UObject setup required before the build (before PostEditChange), the UObject dependencies should exist and have all the source data and properties imported.
- * 5. FinalizeObject_GameThread - Do any final UObject setup after the build (after PostEditChange). Note that the build of an asset can be asynchronous and not finish yet.
+ * 5. FinalizeObject_GameThread - Do any final UObject setup after the build (after PostEditChange). Note that the build of an asset can be asynchronous and this function will be call after the async build is done.
  * 
  * 
  * Scene factory implementation
diff --git a/Engine/Source/Runtime/Interchange/Engine/Private/InterchangeManager.cpp b/Engine/Source/Runtime/Interchange/Engine/Private/InterchangeManager.cpp
--- a/Engine/Source/Runtime/Interchange/Engine/Private/InterchangeManager.cpp
+++ b/Engine/Source/Runtime/Interchange/Engine/Private/InterchangeManager.cpp
@@ -769,6 +769,14 @@ void UE::Interchange::FImportResult::AddImportedObject(UObject* ImportedObject)
 void UE::Interchange::FImportResult::OnDone(TFunction< void(FImportResult&) > Callback)
 {
 	DoneCallback = Callback;
+	//In case the import is already done (because it was synchronous) execute the new OnDone callback
+	if (ImportStatus == EStatus::Done)
+	{
+		if (DoneCallback)
+		{
+			DoneCallback(*this);
+		}
+	}
 }
 
 void UE::Interchange::FImportResult::AddReferencedObjects(FReferenceCollector& Collector)
@@ -961,6 +969,13 @@ UInterchangeManager& UInterchangeManager::GetInterchangeManager()
 			//Task should have been cancel in the Engine pre exit callback
 			ensure(InterchangeManager->ImportTasks.Num() == 0);
 			InterchangeManager->OnPreDestroyInterchangeManager.Broadcast();
+
+			if (InterchangeManager->QueuedPostImportTasksTickerHandle.IsValid())
+			{
+				FTSTicker::GetCoreTicker().RemoveTicker(InterchangeManager->QueuedPostImportTasksTickerHandle);
+				InterchangeManager->QueuedPostImportTasksTickerHandle.Reset();
+			}
+
 			//Release the InterchangeManager object
 			InterchangeManager.Reset();
 			InterchangeManagerScopeOfLifeEnded = true;
@@ -2021,6 +2036,49 @@ bool UInterchangeManager::IsObjectBeingImported(UObject* Object) const
 	return false;
 }
 
+bool UInterchangeManager::EnqueuePostImportTask(TSharedPtr<FInterchangePostImportTask> PostImportTask)
+{
+	//We can only enqueue on the game thread
+	if (!ensure(IsInGameThread()))
+	{
+		return false;
+	}
+
+	QueuedPostImportTasks.Enqueue(PostImportTask);
+
+	if (!QueuedPostImportTasksTickerHandle.IsValid())
+	{
+		QueuedPostImportTasksTickerHandle = FTSTicker::GetCoreTicker().AddTicker(FTickerDelegate::CreateLambda([this](float DeltaTime)
+			{
+				check(IsInGameThread());
+				while (!QueuedPostImportTasks.IsEmpty())
+				{
+					//Wait next frame if we are importing assets or scenes
+					if (!QueuedTasks.IsEmpty() || ImportTasks.Num() > 0)
+					{
+						break;
+					}
+					TSharedPtr<FInterchangePostImportTask> PostImportTask;
+					if (QueuedPostImportTasks.Dequeue(PostImportTask))
+					{
+						if (PostImportTask)
+						{
+							PostImportTask->Execute();
+						}
+					}
+				}
+
+				if (QueuedPostImportTasks.IsEmpty())
+				{
+					QueuedPostImportTasksTickerHandle.Reset();
+					return false;
+				}
+				return true;
+			}));
+	}
+	return true;
+}
+
 bool UInterchangeManager::IsInterchangeImportEnabled()
 {
 	return CCvarInterchangeImportEnable->GetBool();
diff --git a/Engine/Source/Runtime/Interchange/Engine/Private/InterchangeMeshUtilities.cpp b/Engine/Source/Runtime/Interchange/Engine/Private/InterchangeMeshUtilities.cpp
--- a/Engine/Source/Runtime/Interchange/Engine/Private/InterchangeMeshUtilities.cpp
+++ b/Engine/Source/Runtime/Interchange/Engine/Private/InterchangeMeshUtilities.cpp
@@ -24,7 +24,7 @@
 
 #include UE_INLINE_GENERATED_CPP_BY_NAME(InterchangeMeshUtilities)
 
-TFuture<bool> UInterchangeMeshUtilities::ImportCustomLodAsync(UObject* MeshObject, const int32 LodIndex)
+TFuture<bool> UInterchangeMeshUtilities::ImportCustomLod(UObject* MeshObject, const int32 LodIndex)
 {
 	TSharedPtr<TPromise<bool>> Promise = MakeShared<TPromise<bool>>();
 	if (!MeshObject)
@@ -69,7 +69,7 @@ TFuture<bool> UInterchangeMeshUtilities::ImportCustomLodAsync(UObject* MeshObjec
 			if (ensure(Filenames.Num() == 1))
 			{
 				const UInterchangeSourceData* SourceData = InterchangeManager.CreateSourceData(Filenames[0]);
-				return InternalImportCustomLodAsync(Promise, MeshObject, LodIndex, SourceData);
+				return InternalImportCustomLod(Promise, MeshObject, LodIndex, SourceData);
 			}
 		}
 	}
@@ -78,14 +78,14 @@ TFuture<bool> UInterchangeMeshUtilities::ImportCustomLodAsync(UObject* MeshObjec
 	return Promise->GetFuture();
 }
 
-TFuture<bool> UInterchangeMeshUtilities::ImportCustomLodAsync(UObject* MeshObject, const int32 LodIndex, const UInterchangeSourceData* SourceData)
+TFuture<bool> UInterchangeMeshUtilities::ImportCustomLod(UObject* MeshObject, const int32 LodIndex, const UInterchangeSourceData* SourceData)
 {
 	TSharedPtr<TPromise<bool>> Promise = MakeShared<TPromise<bool>>();
 	
-	return InternalImportCustomLodAsync(Promise, MeshObject, LodIndex, SourceData);
+	return InternalImportCustomLod(Promise, MeshObject, LodIndex, SourceData);
 }
 
-TFuture<bool> UInterchangeMeshUtilities::InternalImportCustomLodAsync(TSharedPtr<TPromise<bool>> Promise, UObject* MeshObject, const int32 LodIndex, const UInterchangeSourceData* SourceData)
+TFuture<bool> UInterchangeMeshUtilities::InternalImportCustomLod(TSharedPtr<TPromise<bool>> Promise, UObject* MeshObject, const int32 LodIndex, const UInterchangeSourceData* SourceData)
 {
 #if WITH_EDITOR
 	UInterchangeManager& InterchangeManager = UInterchangeManager::GetInterchangeManager();
@@ -128,7 +128,7 @@ TFuture<bool> UInterchangeMeshUtilities::InternalImportCustomLodAsync(TSharedPtr
 
 	if (bInvalidLodIndex)
 	{
-		UE_LOG(LogInterchangeEngine, Warning, TEXT("FInterchangeMeshUtilities::InternalImportCustomLodAsync: Invalid mesh LOD index %d, no prior LOD index exists."), LodIndex);
+		UE_LOG(LogInterchangeEngine, Warning, TEXT("FInterchangeMeshUtilities::InternalImportCustomLod: Invalid mesh LOD index %d, no prior LOD index exists."), LodIndex);
 		Promise->SetValue(false);
 		return Promise->GetFuture();
 	}
diff --git a/Engine/Source/Runtime/Interchange/Engine/Private/Tasks/InterchangeTaskCompletion.cpp b/Engine/Source/Runtime/Interchange/Engine/Private/Tasks/InterchangeTaskCompletion.cpp
--- a/Engine/Source/Runtime/Interchange/Engine/Private/Tasks/InterchangeTaskCompletion.cpp
+++ b/Engine/Source/Runtime/Interchange/Engine/Private/Tasks/InterchangeTaskCompletion.cpp
@@ -140,7 +140,7 @@ void UE::Interchange::FTaskPreCompletion::DoTask(ENamedThreads::Type CurrentThre
 		}
 #endif //WITH_EDITOR
 
-		//Third iteration to call FinalizeObject_GameThread
+		//Third iteration to register the assets
 		for (const FImportAsyncHelper::FImportedObjectInfo& ObjectInfo : ImportedObjects)
 		{
 			UObject* ImportedObject = ObjectInfo.ImportedObject;
@@ -148,7 +148,7 @@ void UE::Interchange::FTaskPreCompletion::DoTask(ENamedThreads::Type CurrentThre
 			{
 				continue;
 			}
-			//Post import broadcast
+			//Register the assets
 			if (bIsAsset)
 			{
 				AsyncHelper->AssetImportResult->AddImportedObject(ImportedObject);
diff --git a/Engine/Source/Runtime/Interchange/Engine/Public/InterchangeManager.h b/Engine/Source/Runtime/Interchange/Engine/Public/InterchangeManager.h
--- a/Engine/Source/Runtime/Interchange/Engine/Public/InterchangeManager.h
+++ b/Engine/Source/Runtime/Interchange/Engine/Public/InterchangeManager.h
@@ -13,6 +13,7 @@
 #include "HAL/Thread.h"
 #include "HAL/ThreadSafeBool.h"
 #include "InterchangeAssetImportData.h"
+#include "InterchangeFactoryBase.h"
 #include "InterchangePipelineConfigurationBase.h"
 #include "InterchangeResultsContainer.h"
 #include "InterchangeSourceData.h"
@@ -30,7 +31,6 @@
 
 class FAsyncTaskNotification;
 class UInterchangeBlueprintPipelineBase;
-class UInterchangeFactoryBase;
 class UInterchangeFactoryBaseNode;
 class UInterchangePipelineBase;
 class UInterchangePythonPipelineBase;
@@ -651,6 +651,12 @@ public:
 	 */
 	INTERCHANGEENGINE_API bool IsObjectBeingImported(UObject* Object) const;
 
+	/**
+	 * Queue task that are not directly a import or re-import of assets.
+	 * The post import tasks are execute only when the QueuedTasks is empty (no more import task running)
+	 * Example: if a skeletal mesh re-import cannot apply the existing alternate skinning data, it will enqueue a post import task to re-import those alternate skinning files.
+	 */
+	INTERCHANGEENGINE_API bool EnqueuePostImportTask(TSharedPtr<FInterchangePostImportTask> PostImportTask);
 protected:
 
 	/** Return true if we can show some UI */
@@ -707,10 +713,18 @@ private:
 	TMap<UClass*, bool> NonParallelTranslatorLocks;
 	TMap<UClass*, TArray<FQueuedTaskData>> NonParallelTranslatorQueueTasks;
 	
-	//Queue all incomming tasks if there is more started task then we have cores
+	//Queue all incoming tasks if there is more started task then we have cores
 	TQueue<FQueuedTaskData> QueuedTasks;
 	int32 QueueTaskCount = 0;
 
+	//Queue all incoming task that are not direct import or re-import of assets.
+	// The post import tasks are execute only when all the QueuedTasks are completed and empty
+	//Those task can start an import task (QueuedTasks), one example is the skeletal mesh alternate skinning which require to import several file for each profiles
+	TQueue<TSharedPtr<FInterchangePostImportTask>> QueuedPostImportTasks;
+
+	//Ticker which is on only if we have some QueuedPostImportTasks
+	FTSTicker::FDelegateHandle	QueuedPostImportTasksTickerHandle;
+
 	//By using pointer, there is no issue if the array get resize
 	TArray<TSharedPtr<UE::Interchange::FImportAsyncHelper, ESPMode::ThreadSafe> > ImportTasks;
 
diff --git a/Engine/Source/Runtime/Interchange/Engine/Public/InterchangeMeshUtilities.h b/Engine/Source/Runtime/Interchange/Engine/Public/InterchangeMeshUtilities.h
--- a/Engine/Source/Runtime/Interchange/Engine/Public/InterchangeMeshUtilities.h
+++ b/Engine/Source/Runtime/Interchange/Engine/Public/InterchangeMeshUtilities.h
@@ -24,7 +24,7 @@ public:
 	 *
 	 * @Note - This function will sort out the SourceData and call ImportCustomLodAsync and wait until the result is available.
 	 */
-	static INTERCHANGEENGINE_API TFuture<bool> ImportCustomLodAsync(UObject* MeshObject, const int32 LodIndex);
+	static INTERCHANGEENGINE_API TFuture<bool> ImportCustomLod(UObject* MeshObject, const int32 LodIndex);
 
 
 	/**
@@ -38,11 +38,11 @@ public:
 	 * @Note - Since this is the async version, the SourceData parameter is required.
 	 * @Note - This function will search for an available interchange asset factory that can do the job for the MeshObject class
 	 */
-	static INTERCHANGEENGINE_API TFuture<bool> ImportCustomLodAsync(UObject* MeshObject, const int32 LodIndex, const UInterchangeSourceData* SourceData);
+	static INTERCHANGEENGINE_API TFuture<bool> ImportCustomLod(UObject* MeshObject, const int32 LodIndex, const UInterchangeSourceData* SourceData);
 
 private:
 
-	static TFuture<bool> InternalImportCustomLodAsync(TSharedPtr<TPromise<bool>> Promise, UObject* MeshObject, const int32 LodIndex, const UInterchangeSourceData* SourceData);
+	static TFuture<bool> InternalImportCustomLod(TSharedPtr<TPromise<bool>> Promise, UObject* MeshObject, const int32 LodIndex, const UInterchangeSourceData* SourceData);
 };
 
 
