Skelmesh Editor: Avoid re-registering the viewport menu.

#rnx

diff --git a/Engine/Source/Editor/SkeletalMeshEditor/Private/SkeletalMeshEditor.cpp b/Engine/Source/Editor/SkeletalMeshEditor/Private/SkeletalMeshEditor.cpp
--- a/Engine/Source/Editor/SkeletalMeshEditor/Private/SkeletalMeshEditor.cpp
+++ b/Engine/Source/Editor/SkeletalMeshEditor/Private/SkeletalMeshEditor.cpp
@@ -79,6 +79,10 @@ const FName SkeletalMeshEditorTabs::MorphTargetsTab("MorphTargetsTab");
 const FName SkeletalMeshEditorTabs::CurveMetadataTab(TEXT("AnimCurveMetadataEditorTab"));
 const FName SkeletalMeshEditorTabs::FindReplaceTab("FindReplaceTab");
 
+
+static const FName ViewportMenuId("SkeletalMeshEditor.MeshContextMenu");
+
+
 DEFINE_LOG_CATEGORY(LogSkeletalMeshEditor);
 
 #define LOCTEXT_NAMESPACE "SkeletalMeshEditor"
@@ -1033,7 +1037,7 @@ void FSkeletalMeshEditor::ExtendMenu()
 	{
 		FToolMenuSection& AssetSection = AssetMenu->FindOrAddSection("AssetEditorActions");
 		const FName SkeletalMeshToolkitName = GetToolkitFName();
-		FToolMenuEntry& Entry = AssetSection.AddDynamicEntry("AssetManagerEditorSkeletalMeshCommands", FNewToolMenuSectionDelegate::CreateLambda([SkeletalMeshToolkitName](FToolMenuSection& InSection)
+		AssetSection.AddDynamicEntry("AssetManagerEditorSkeletalMeshCommands", FNewToolMenuSectionDelegate::CreateLambda([SkeletalMeshToolkitName](FToolMenuSection& InSection)
 			{
 				UAssetEditorToolkitMenuContext* MenuContext = InSection.FindContext<UAssetEditorToolkitMenuContext>();
 				if (MenuContext && MenuContext->Toolkit.IsValid() && MenuContext->Toolkit.Pin()->GetToolkitFName() == SkeletalMeshToolkitName)
@@ -1044,8 +1048,10 @@ void FSkeletalMeshEditor::ExtendMenu()
 		));
 	}
 
-	UToolMenu* ViewportMenu = UToolMenus::Get()->RegisterMenu("SkeletalMeshEditor.MeshContextMenu");
+
+	if (!UToolMenus::Get()->IsMenuRegistered(ViewportMenuId))
 	{
+		UToolMenu* ViewportMenu = UToolMenus::Get()->RegisterMenu(ViewportMenuId);
 		FToolMenuSection& AssetSection = ViewportMenu->AddSection("Asset", LOCTEXT("MeshClickMenu_Section_Asset", "Asset"));
 		AssetSection.AddDynamicEntry("DynamicAssetEntries", FNewToolMenuSectionDelegate::CreateLambda([](FToolMenuSection& InSection)
 		{
@@ -1489,7 +1495,7 @@ void FSkeletalMeshEditor::HandleMeshClick(HActor* HitProxy, const FViewportClick
 			FSlateApplication::Get().PushMenu(
 				FSlateApplication::Get().GetActiveTopLevelWindow().ToSharedRef(),
 				FWidgetPath(),
-				UToolMenus::Get()->GenerateWidget("SkeletalMeshEditor.MeshContextMenu", ToolMenuContext),
+				UToolMenus::Get()->GenerateWidget(ViewportMenuId, ToolMenuContext),
 				FSlateApplication::Get().GetCursorPos(),
 				FPopupTransitionEffect(FPopupTransitionEffect::ContextMenu)
 				);
