Add new reimport button in meshes editor to force re-import with dialog.

#rnx
#prefligth 66268cac684c0e6f587a5e3d

diff --git a/Engine/Source/Editor/AnimationEditor/Private/AnimationEditor.cpp b/Engine/Source/Editor/AnimationEditor/Private/AnimationEditor.cpp
--- a/Engine/Source/Editor/AnimationEditor/Private/AnimationEditor.cpp
+++ b/Engine/Source/Editor/AnimationEditor/Private/AnimationEditor.cpp
@@ -244,7 +244,11 @@ void FAnimationEditor::BindCommands()
 		FCanExecuteAction::CreateSP(this, &FAnimationEditor::CanSetKey));
 
 	ToolkitCommands->MapAction(FAnimationEditorCommands::Get().ReimportAnimation,
-		FExecuteAction::CreateSP(this, &FAnimationEditor::OnReimportAnimation),
+		FExecuteAction::CreateSP(this, &FAnimationEditor::OnReimportAnimation, false),
+		FCanExecuteAction::CreateSP(this, &FAnimationEditor::HasValidAnimationSequence));
+
+	ToolkitCommands->MapAction(FAnimationEditorCommands::Get().ReimportAnimationWithDialog,
+		FExecuteAction::CreateSP(this, &FAnimationEditor::OnReimportAnimation, true),
 		FCanExecuteAction::CreateSP(this, &FAnimationEditor::HasValidAnimationSequence));
 
 	ToolkitCommands->MapAction(FAnimationEditorCommands::Get().ExportToFBX_AnimData,
@@ -314,6 +318,7 @@ void FAnimationEditor::ExtendToolbar()
 	{
 		FToolMenuSection& AnimationSection = ToolMenu->AddSection("Animation", LOCTEXT("ToolbarAnimationSectionLabel", "Animation"), SectionInsertLocation);
 		AnimationSection.AddEntry(FToolMenuEntry::InitToolBarButton(FAnimationEditorCommands::Get().ReimportAnimation));
+		AnimationSection.AddEntry(FToolMenuEntry::InitToolBarButton(FAnimationEditorCommands::Get().ReimportAnimationWithDialog));
 		AnimationSection.AddEntry(FToolMenuEntry::InitToolBarButton(FAnimationEditorCommands::Get().ApplyCompression, LOCTEXT("Toolbar_ApplyCompression", "Apply Compression")));
 		AnimationSection.AddEntry(FToolMenuEntry::InitComboButton(
 			"ExportAsset",
@@ -639,12 +644,28 @@ void FAnimationEditor::OnSetKey()
 	}
 }
 
-void FAnimationEditor::OnReimportAnimation()
+void FAnimationEditor::OnReimportAnimation(bool bWithDialog)
 {
 	UAnimSequence* AnimSequence = Cast<UAnimSequence>(AnimationAsset);
 	if (AnimSequence)
 	{
-		FReimportManager::Instance()->ReimportAsync(AnimSequence, true);
+		constexpr bool bAskForNewFileIfMissingTrue = true;
+		constexpr bool bShowNotificationTrue = true;
+		const FString PreferredReimportFileEmpty = TEXT("");
+		constexpr FReimportHandler* SpecifiedReimportHandlerNull = nullptr;
+		constexpr bool bForceNewFileFalse = false;
+		constexpr int32 SourceFileIndex = INDEX_NONE;
+		constexpr bool bAutomatedFalse = false;
+
+		FReimportManager::Instance()->ReimportAsync(AnimSequence
+			, bAskForNewFileIfMissingTrue
+			, bShowNotificationTrue
+			, PreferredReimportFileEmpty
+			, SpecifiedReimportHandlerNull
+			, SourceFileIndex
+			, bForceNewFileFalse
+			, bAutomatedFalse
+			, bWithDialog);
 	}
 }
 
diff --git a/Engine/Source/Editor/AnimationEditor/Private/AnimationEditor.h b/Engine/Source/Editor/AnimationEditor/Private/AnimationEditor.h
--- a/Engine/Source/Editor/AnimationEditor/Private/AnimationEditor.h
+++ b/Engine/Source/Editor/AnimationEditor/Private/AnimationEditor.h
@@ -139,7 +139,7 @@ private:
 
 	void OnSetKey();
 
-	void OnReimportAnimation();
+	void OnReimportAnimation(bool bWithDialog);
 
 	void OnApplyCompression();
 
diff --git a/Engine/Source/Editor/AnimationEditor/Private/AnimationEditorCommands.cpp b/Engine/Source/Editor/AnimationEditor/Private/AnimationEditorCommands.cpp
--- a/Engine/Source/Editor/AnimationEditor/Private/AnimationEditorCommands.cpp
+++ b/Engine/Source/Editor/AnimationEditor/Private/AnimationEditorCommands.cpp
@@ -12,6 +12,7 @@ void FAnimationEditorCommands::RegisterCommands()
 {
 	UI_COMMAND(ImportAnimation, "Import Animation", "Import new animation for the skeleton.", EUserInterfaceActionType::Button, FInputChord());
 	UI_COMMAND(ReimportAnimation, "Reimport Animation", "Reimport current animation.", EUserInterfaceActionType::Button, FInputChord());
+	UI_COMMAND(ReimportAnimationWithDialog, "Reimport Animation With Dialog", "Reimport current animation and show the re-import dialog.", EUserInterfaceActionType::Button, FInputChord());
 	UI_COMMAND(ApplyCompression, "Apply Compression", "Apply compression to current animation", EUserInterfaceActionType::Button, FInputChord());
 	UI_COMMAND(ExportToFBX_AnimData, "Animation Data", "Export current animation source data to FBX", EUserInterfaceActionType::Button, FInputChord());
 	UI_COMMAND(ExportToFBX_PreviewMesh, "Preview Mesh", "Export current animation by playing on the current preview mesh, including retargeting, post process graph, or anything you see on the preview mesh.", EUserInterfaceActionType::Button, FInputChord());
diff --git a/Engine/Source/Editor/AnimationEditor/Private/AnimationEditorCommands.h b/Engine/Source/Editor/AnimationEditor/Private/AnimationEditorCommands.h
--- a/Engine/Source/Editor/AnimationEditor/Private/AnimationEditorCommands.h
+++ b/Engine/Source/Editor/AnimationEditor/Private/AnimationEditorCommands.h
@@ -27,6 +27,7 @@ public:
 	TSharedPtr<FUICommandInfo> ImportAnimation;
 	// reimport animation
 	TSharedPtr<FUICommandInfo> ReimportAnimation;
+	TSharedPtr<FUICommandInfo> ReimportAnimationWithDialog;
 	// apply compression
 	TSharedPtr<FUICommandInfo> ApplyCompression;
 	// export to FBX
diff --git a/Engine/Source/Editor/EditorStyle/Private/StarshipStyle.cpp b/Engine/Source/Editor/EditorStyle/Private/StarshipStyle.cpp
--- a/Engine/Source/Editor/EditorStyle/Private/StarshipStyle.cpp
+++ b/Engine/Source/Editor/EditorStyle/Private/StarshipStyle.cpp
@@ -5060,6 +5060,7 @@ void FStarshipEditorStyle::FStyle::SetupPersonaStyle()
 		Set("AnimationEditor.ApplyCompression", new IMAGE_BRUSH_SVG("Starship/Common/Compress", Icon20x20));
 		Set("AnimationEditor.ExportToFBX", new CORE_IMAGE_BRUSH_SVG("Starship/Common/export_20", Icon20x20));
 		Set("AnimationEditor.ReimportAnimation", new CORE_IMAGE_BRUSH_SVG("Starship/Common/import", Icon20x20));
+		Set("AnimationEditor.ReimportAnimationWithDialog", new CORE_IMAGE_BRUSH_SVG("Starship/Common/import", Icon20x20));
 		Set("AnimationEditor.CreateAsset", new IMAGE_BRUSH_SVG("Starship/Persona/PersonaCreateAsset", Icon20x20));
 		Set("AnimationEditor.SetKey", new CORE_IMAGE_BRUSH_SVG("Starship/Common/plus", Icon20x20));
 		Set("AnimationEditor.ApplyAnimation", new IMAGE_BRUSH_SVG("Starship/Common/Apply", Icon20x20));
@@ -5085,6 +5086,7 @@ void FStarshipEditorStyle::FStyle::SetupPersonaStyle()
 	// Skeletal mesh editor
 	{
 		Set("SkeletalMeshEditor.ReimportMesh", new CORE_IMAGE_BRUSH_SVG("Starship/Common/import", Icon20x20));
+		Set("SkeletalMeshEditor.ReimportWithDialog", new CORE_IMAGE_BRUSH_SVG("Starship/Common/import", Icon20x20));
 		Set("SkeletalMeshEditor.ImportLODs", new CORE_IMAGE_BRUSH_SVG("Starship/Common/import", Icon20x20));
 
 		Set("SkeletalMeshEditor.MeshSectionSelection", new IMAGE_BRUSH_SVG("Starship/Persona/SectionSelection", Icon20x20));
@@ -6877,6 +6879,7 @@ void FStarshipEditorStyle::FStyle::SetupToolkitStyles()
 		Set("StaticMeshEditor.SetShowVertexColor", new IMAGE_BRUSH_SVG( "Starship/Common/SetShowVertexColors", Icon16x16));
 		Set("StaticMeshEditor.SetRealtimePreview", new IMAGE_BRUSH_SVG("Starship/Common/Realtime", Icon16x16));
 		Set("StaticMeshEditor.ReimportMesh", new IMAGE_BRUSH_SVG("Starship/StaticMeshEditor/ReimportMesh", Icon20x20));
+		Set("StaticMeshEditor.ReimportMeshWithDialog", new IMAGE_BRUSH_SVG("Starship/StaticMeshEditor/ReimportMesh", Icon20x20));
 		Set("StaticMeshEditor.SetShowBounds", new IMAGE_BRUSH_SVG("Starship/Common/SetShowBounds", Icon16x16));
 		Set("StaticMeshEditor.SetDrawUVs", new IMAGE_BRUSH_SVG("Starship/Common/SetDrawUVs", Icon20x20));
 		Set("StaticMeshEditor.SetShowCollision", new IMAGE_BRUSH_SVG("Starship/Common/Collision", Icon20x20));
diff --git a/Engine/Source/Editor/SkeletalMeshEditor/Private/SkeletalMeshEditor.cpp b/Engine/Source/Editor/SkeletalMeshEditor/Private/SkeletalMeshEditor.cpp
--- a/Engine/Source/Editor/SkeletalMeshEditor/Private/SkeletalMeshEditor.cpp
+++ b/Engine/Source/Editor/SkeletalMeshEditor/Private/SkeletalMeshEditor.cpp
@@ -305,11 +305,18 @@ void FSkeletalMeshEditor::BindCommands()
 {
 	FSkeletalMeshEditorCommands::Register();
 
+	constexpr bool bWithNewFileFalse = false;
+	constexpr bool bWithDialogFalse = false;
+	constexpr bool bWithDialogTrue = true;
+
 	ToolkitCommands->MapAction(FSkeletalMeshEditorCommands::Get().ReimportMesh,
-		FExecuteAction::CreateSP(this, &FSkeletalMeshEditor::HandleReimportMesh, (int32)INDEX_NONE));
+		FExecuteAction::CreateSP(this, &FSkeletalMeshEditor::HandleReimportMesh, FReimportParameters((int32)INDEX_NONE, bWithNewFileFalse, bWithDialogFalse)));
 
 	ToolkitCommands->MapAction(FSkeletalMeshEditorCommands::Get().ReimportAllMesh,
-		FExecuteAction::CreateSP(this, &FSkeletalMeshEditor::HandleReimportAllMesh, (int32)INDEX_NONE));
+		FExecuteAction::CreateSP(this, &FSkeletalMeshEditor::HandleReimportAllMesh, FReimportParameters((int32)INDEX_NONE, bWithNewFileFalse, bWithDialogFalse)));
+
+	ToolkitCommands->MapAction(FSkeletalMeshEditorCommands::Get().ReimportWithDialog,
+		FExecuteAction::CreateSP(this, &FSkeletalMeshEditor::HandleReimportMesh, FReimportParameters((int32)INDEX_NONE, bWithNewFileFalse, bWithDialogTrue)));
 
 	ToolkitCommands->MapAction(FPersonaCommonCommands::Get().TogglePlay,
 		FExecuteAction::CreateRaw(&GetPersonaToolkit()->GetPreviewScene().Get(), &IPersonaPreviewScene::TogglePlayback));
@@ -333,10 +340,11 @@ TSharedPtr<FSkeletalMeshEditor> FSkeletalMeshEditor::GetSkeletalMeshEditor(const
 	return TSharedPtr<FSkeletalMeshEditor>();
 }
 
-void FSkeletalMeshEditor::RegisterReimportContextMenu(const FName InBaseMenuName)
+void FSkeletalMeshEditor::RegisterReimportContextMenu(const FName InBaseMenuName, bool bWithDialog)
 {
-	static auto ReimportMeshWithNewFileAction = [](const FToolMenuContext& InMenuContext, int32 SourceFileIndex)
+	static auto ReimportMeshWithNewFileAction = [](const FToolMenuContext& InMenuContext, int32 SourceFileIndex, bool bWithDialog)
 	{
+		constexpr bool bWithNewFileTrue = true;
 		TSharedPtr<FSkeletalMeshEditor> SkeletalMeshEditor = GetSkeletalMeshEditor(InMenuContext);
 		if (SkeletalMeshEditor.IsValid())
 		{
@@ -345,42 +353,14 @@ void FSkeletalMeshEditor::RegisterReimportContextMenu(const FName InBaseMenuName
 				if (UFbxSkeletalMeshImportData* SkeletalMeshImportData = Cast<UFbxSkeletalMeshImportData>(FoundSkeletalMesh->GetAssetImportData()))
 				{
 					SkeletalMeshImportData->ImportContentType = SourceFileIndex == 0 ? EFBXImportContentType::FBXICT_All : SourceFileIndex == 1 ? EFBXImportContentType::FBXICT_Geometry : EFBXImportContentType::FBXICT_SkinningWeights;
-					SkeletalMeshEditor->HandleReimportMeshWithNewFile(SourceFileIndex);
-				}
-				else if (UInterchangeAssetImportData* InterchangeImportData = Cast<UInterchangeAssetImportData>(FoundSkeletalMesh->GetAssetImportData()))
-				{
-					SkeletalMeshEditor->HandleReimportMeshWithNewFile(SourceFileIndex);
 				}
+				
+				SkeletalMeshEditor->HandleReimportMesh(FReimportParameters(SourceFileIndex, bWithNewFileTrue, bWithDialog));
 			}
 		}
 	};
 
-	static auto CreateMultiContentSubMenu = [](UToolMenu* InMenu)
-	{
-		FToolMenuSection& Section = InMenu->AddSection("Reimport");
-
-		TSharedPtr<FSkeletalMeshEditor> SkeletalMeshEditor = GetSkeletalMeshEditor(InMenu->Context);
-		if (SkeletalMeshEditor.IsValid())
-		{
-			Section.AddMenuEntry(
-				"ReimportGeometryContentLabel",
-				LOCTEXT("ReimportGeometryContentLabel", "Geometry"),
-				LOCTEXT("ReimportGeometryContentLabelTooltipTooltip", "Reimport Geometry Only"),
-				FSlateIcon(FAppStyle::GetAppStyleSetName(), "ContentBrowser.AssetActions.ReimportAsset"),
-				FToolMenuExecuteAction::CreateLambda(ReimportMeshWithNewFileAction, 1)
-			);
-
-			Section.AddMenuEntry(
-				"ReimportSkinningAndWeightsContentLabel",
-				LOCTEXT("ReimportSkinningAndWeightsContentLabel", "Skinning And Weights"),
-				LOCTEXT("ReimportSkinningAndWeightsContentLabelTooltipTooltip", "Reimport Skinning And Weights Only"),
-				FSlateIcon(FAppStyle::GetAppStyleSetName(), "ContentBrowser.AssetActions.ReimportAsset"),
-				FToolMenuExecuteAction::CreateLambda(ReimportMeshWithNewFileAction, 2)
-			);
-		}
-	};
-
-	static auto ReimportAction = [](const FToolMenuContext& InMenuContext, const int32 SourceFileIndex, const bool bReimportAll, const bool bWithNewFile)
+	static auto ReimportAction = [](const FToolMenuContext& InMenuContext, const int32 SourceFileIndex, const bool bReimportAll, const bool bWithNewFile, const bool bWithDialog)
 	{
 		TSharedPtr<FSkeletalMeshEditor> SkeletalMeshEditor = GetSkeletalMeshEditor(InMenuContext);
 		if (SkeletalMeshEditor.IsValid())
@@ -395,31 +375,17 @@ void FSkeletalMeshEditor::RegisterReimportContextMenu(const FName InBaseMenuName
 
 				if (bReimportAll)
 				{
-					if (bWithNewFile)
-					{
-						SkeletalMeshEditor->HandleReimportAllMeshWithNewFile(SourceFileIndex);
-					}
-					else
-					{
-						SkeletalMeshEditor->HandleReimportAllMesh(SourceFileIndex);
-					}
+					SkeletalMeshEditor->HandleReimportAllMesh(FReimportParameters(SourceFileIndex, bWithNewFile, bWithDialog));
 				}
 				else
 				{
-					if (bWithNewFile)
-					{
-						SkeletalMeshEditor->HandleReimportMeshWithNewFile(SourceFileIndex);
-					}
-					else
-					{
-						SkeletalMeshEditor->HandleReimportMesh(SourceFileIndex);
-					}
+					SkeletalMeshEditor->HandleReimportMesh(FReimportParameters(SourceFileIndex, bWithNewFile, bWithDialog));
 				}
 			}
 		}
 	};
 
-	static auto CreateReimportSubMenu = [](UToolMenu* InMenu, bool bReimportAll, bool bWithNewFile)
+	static auto CreateReimportSubMenu = [](UToolMenu* InMenu, bool bReimportAll, bool bWithNewFile, bool bWithDialog)
 	{
 		TSharedPtr<FSkeletalMeshEditor> SkeletalMeshEditor = GetSkeletalMeshEditor(InMenu->Context);
 		if (SkeletalMeshEditor.IsValid())
@@ -451,7 +417,7 @@ void FSkeletalMeshEditor::RegisterReimportContextMenu(const FName InBaseMenuName
 							ReimportLabel,
 							ReimportLabelTooltip,
 							FSlateIcon(FAppStyle::GetAppStyleSetName(), "ContentBrowser.AssetActions.ReimportAsset"),
-							FToolMenuExecuteAction::CreateLambda(ReimportAction, SourceFileIndex, bReimportAll, bWithNewFile)
+							FToolMenuExecuteAction::CreateLambda(ReimportAction, SourceFileIndex, bReimportAll, bWithNewFile, bWithDialog)
 						);
 					}
 				}
@@ -459,10 +425,10 @@ void FSkeletalMeshEditor::RegisterReimportContextMenu(const FName InBaseMenuName
 		}
 	};
 
-	if (!UToolMenus::Get()->IsMenuRegistered(UToolMenus::JoinMenuPaths(InBaseMenuName, "ReimportContextMenu")))
+	if (!UToolMenus::Get()->IsMenuRegistered(UToolMenus::JoinMenuPaths(InBaseMenuName, bWithDialog ? "ReimportWithDialogContextMenu" : "ReimportContextMenu")))
 	{
-		UToolMenu* ToolMenu = UToolMenus::Get()->RegisterMenu(UToolMenus::JoinMenuPaths(InBaseMenuName, "ReimportContextMenu"));
-		ToolMenu->AddDynamicSection("Section", FNewToolMenuDelegate::CreateLambda([](UToolMenu* InMenu)
+		UToolMenu* ToolMenu = UToolMenus::Get()->RegisterMenu(UToolMenus::JoinMenuPaths(InBaseMenuName, bWithDialog ? "ReimportWithDialogContextMenu" : "ReimportContextMenu"));
+		ToolMenu->AddDynamicSection("Section", FNewToolMenuDelegate::CreateLambda([bWithDialog](UToolMenu* InMenu)
 		{
 			TSharedPtr<FSkeletalMeshEditor> SkeletalMeshEditor = GetSkeletalMeshEditor(InMenu->Context);
 			if (SkeletalMeshEditor.IsValid())
@@ -470,6 +436,9 @@ void FSkeletalMeshEditor::RegisterReimportContextMenu(const FName InBaseMenuName
 				USkeletalMesh* InSkeletalMesh = SkeletalMeshEditor->SkeletalMesh;
 				bool bShowSubMenu = InSkeletalMesh != nullptr && InSkeletalMesh->GetAssetImportData() != nullptr && InSkeletalMesh->GetAssetImportData()->GetSourceFileCount() > 1;
 
+				constexpr bool bWithNewFileFalse = false;
+				constexpr bool bWithNewFileTrue = true;
+
 				FToolMenuSection& Section = InMenu->AddSection("Section");
 				if (!bShowSubMenu)
 				{
@@ -479,14 +448,14 @@ void FSkeletalMeshEditor::RegisterReimportContextMenu(const FName InBaseMenuName
 						FSkeletalMeshEditorCommands::Get().ReimportMesh->GetLabel(),
 						FSkeletalMeshEditorCommands::Get().ReimportMesh->GetDescription(),
 						FSlateIcon(),
-						FUIAction(FExecuteAction::CreateSP(SkeletalMeshEditor.ToSharedRef(), &FSkeletalMeshEditor::HandleReimportMesh, 0)));
+						FUIAction(FExecuteAction::CreateSP(SkeletalMeshEditor.ToSharedRef(), &FSkeletalMeshEditor::HandleReimportMesh, FReimportParameters(0, bWithNewFileFalse, bWithDialog))));
 
 					Section.AddMenuEntry(
 						FSkeletalMeshEditorCommands::Get().ReimportMeshWithNewFile->GetCommandName(),
 						FSkeletalMeshEditorCommands::Get().ReimportMeshWithNewFile->GetLabel(),
 						FSkeletalMeshEditorCommands::Get().ReimportMeshWithNewFile->GetDescription(),
 						FSlateIcon(),
-						FUIAction(FExecuteAction::CreateSP(SkeletalMeshEditor.ToSharedRef(), &FSkeletalMeshEditor::HandleReimportMeshWithNewFile, 0)));
+						FUIAction(FExecuteAction::CreateSP(SkeletalMeshEditor.ToSharedRef(), &FSkeletalMeshEditor::HandleReimportMesh, FReimportParameters(0, bWithNewFileTrue, bWithDialog))));
 
 					//Reimport ALL
 					Section.AddMenuEntry(
@@ -494,20 +463,43 @@ void FSkeletalMeshEditor::RegisterReimportContextMenu(const FName InBaseMenuName
 						FSkeletalMeshEditorCommands::Get().ReimportAllMesh->GetLabel(),
 						FSkeletalMeshEditorCommands::Get().ReimportAllMesh->GetDescription(),
 						FSlateIcon(),
-						FUIAction(FExecuteAction::CreateSP(SkeletalMeshEditor.ToSharedRef(), &FSkeletalMeshEditor::HandleReimportAllMesh, 0)));
+						FUIAction(FExecuteAction::CreateSP(SkeletalMeshEditor.ToSharedRef(), &FSkeletalMeshEditor::HandleReimportAllMesh, FReimportParameters(0, bWithNewFileFalse, bWithDialog))));
 
 					Section.AddMenuEntry(
 						FSkeletalMeshEditorCommands::Get().ReimportAllMeshWithNewFile->GetCommandName(),
 						FSkeletalMeshEditorCommands::Get().ReimportAllMeshWithNewFile->GetLabel(),
 						FSkeletalMeshEditorCommands::Get().ReimportAllMeshWithNewFile->GetDescription(),
 						FSlateIcon(),
-						FUIAction(FExecuteAction::CreateSP(SkeletalMeshEditor.ToSharedRef(), &FSkeletalMeshEditor::HandleReimportAllMeshWithNewFile, 0)));
+						FUIAction(FExecuteAction::CreateSP(SkeletalMeshEditor.ToSharedRef(), &FSkeletalMeshEditor::HandleReimportAllMesh, FReimportParameters(0, bWithNewFileTrue, bWithDialog))));
 
 					Section.AddSubMenu(
 						"ReimportMultiSources",
 						LOCTEXT("ReimportMultiSources", "Reimport Content"),
 						LOCTEXT("ReimportMultiSourcesTooltip", "Reimport Geometry or Skinning Weights content, this will create multi import source file."),
-						FNewToolMenuDelegate::CreateLambda(CreateMultiContentSubMenu));
+						FNewToolMenuDelegate::CreateLambda([bWithDialog](UToolMenu* InMenu)
+							{
+								FToolMenuSection& Section = InMenu->AddSection("Reimport");
+
+								TSharedPtr<FSkeletalMeshEditor> SkeletalMeshEditor = GetSkeletalMeshEditor(InMenu->Context);
+								if (SkeletalMeshEditor.IsValid())
+								{
+									Section.AddMenuEntry(
+										"ReimportGeometryContentLabel",
+										LOCTEXT("ReimportGeometryContentLabel", "Geometry"),
+										LOCTEXT("ReimportGeometryContentLabelTooltipTooltip", "Reimport Geometry Only"),
+										FSlateIcon(FAppStyle::GetAppStyleSetName(), "ContentBrowser.AssetActions.ReimportAsset"),
+										FToolMenuExecuteAction::CreateLambda(ReimportMeshWithNewFileAction, 1, bWithDialog)
+									);
+
+									Section.AddMenuEntry(
+										"ReimportSkinningAndWeightsContentLabel",
+										LOCTEXT("ReimportSkinningAndWeightsContentLabel", "Skinning And Weights"),
+										LOCTEXT("ReimportSkinningAndWeightsContentLabelTooltipTooltip", "Reimport Skinning And Weights Only"),
+										FSlateIcon(FAppStyle::GetAppStyleSetName(), "ContentBrowser.AssetActions.ReimportAsset"),
+										FToolMenuExecuteAction::CreateLambda(ReimportMeshWithNewFileAction, 2, bWithDialog)
+									);
+								}
+							}));
 				}
 				else
 				{
@@ -516,25 +508,25 @@ void FSkeletalMeshEditor::RegisterReimportContextMenu(const FName InBaseMenuName
 						FSkeletalMeshEditorCommands::Get().ReimportMesh->GetCommandName(),
 						FSkeletalMeshEditorCommands::Get().ReimportMesh->GetLabel(),
 						FSkeletalMeshEditorCommands::Get().ReimportMesh->GetDescription(),
-						FNewToolMenuDelegate::CreateLambda(CreateReimportSubMenu, false, false));
+						FNewToolMenuDelegate::CreateLambda(CreateReimportSubMenu, false, false, bWithDialog));
 			
 					Section.AddSubMenu(
 						FSkeletalMeshEditorCommands::Get().ReimportMeshWithNewFile->GetCommandName(),
 						FSkeletalMeshEditorCommands::Get().ReimportMeshWithNewFile->GetLabel(),
 						FSkeletalMeshEditorCommands::Get().ReimportMeshWithNewFile->GetDescription(),
-						FNewToolMenuDelegate::CreateLambda(CreateReimportSubMenu, false, true));
+						FNewToolMenuDelegate::CreateLambda(CreateReimportSubMenu, false, true, bWithDialog));
 
 					Section.AddSubMenu(
 						FSkeletalMeshEditorCommands::Get().ReimportAllMesh->GetCommandName(),
 						FSkeletalMeshEditorCommands::Get().ReimportAllMesh->GetLabel(),
 						FSkeletalMeshEditorCommands::Get().ReimportAllMesh->GetDescription(),
-						FNewToolMenuDelegate::CreateLambda(CreateReimportSubMenu, true, false));
+						FNewToolMenuDelegate::CreateLambda(CreateReimportSubMenu, true, false, bWithDialog));
 
 					Section.AddSubMenu(
 						FSkeletalMeshEditorCommands::Get().ReimportAllMeshWithNewFile->GetCommandName(),
 						FSkeletalMeshEditorCommands::Get().ReimportAllMeshWithNewFile->GetLabel(),
 						FSkeletalMeshEditorCommands::Get().ReimportAllMeshWithNewFile->GetDescription(),
-						FNewToolMenuDelegate::CreateLambda(CreateReimportSubMenu, true, true));
+						FNewToolMenuDelegate::CreateLambda(CreateReimportSubMenu, true, true, bWithDialog));
 				}
 			}
 		}));
@@ -547,7 +539,11 @@ void FSkeletalMeshEditor::ExtendToolbar()
 	// Add in Editor Specific functionality
 	FName ParentName;
 	static const FName MenuName = GetToolMenuToolbarName(ParentName);
-	RegisterReimportContextMenu(MenuName);
+
+	constexpr bool bWithDialogFalse = false;
+	constexpr bool bWithDialogTrue = true;
+	RegisterReimportContextMenu(MenuName, bWithDialogFalse);
+	RegisterReimportContextMenu(MenuName, bWithDialogTrue);
 
 	UToolMenu* ToolMenu = UToolMenus::Get()->ExtendMenu(MenuName);
 	const FToolMenuInsert SectionInsertLocation("Asset", EToolMenuInsertType::After);
@@ -566,13 +562,17 @@ void FSkeletalMeshEditor::ExtendToolbar()
 		}), SectionInsertLocation);
 	}
 
+	//Reimport button
 	{
 		FToolMenuSection& Section = ToolMenu->AddSection("Mesh", FText(), SectionInsertLocation);
 		
-		FToolMenuEntry Entry = FToolMenuEntry::InitToolBarButton(FSkeletalMeshEditorCommands::Get().ReimportMesh);
 		Section.AddEntry(FToolMenuEntry::InitToolBarButton(FSkeletalMeshEditorCommands::Get().ReimportMesh));
 
 		Section.AddEntry(FToolMenuEntry::InitComboButton("ReimportContextMenu", FUIAction(), FNewToolMenuDelegate(), FText(), FText(), FSlateIcon(), true));
+		
+		Section.AddEntry(FToolMenuEntry::InitToolBarButton(FSkeletalMeshEditorCommands::Get().ReimportWithDialog));
+
+		Section.AddEntry(FToolMenuEntry::InitComboButton("ReimportWithDialogContextMenu", FUIAction(), FNewToolMenuDelegate(), FText(), FText(), FSlateIcon(), true));
 	}
 
 	// If the ToolbarExtender is valid, remove it before rebuilding it
@@ -1274,12 +1274,13 @@ UObject* FSkeletalMeshEditor::HandleGetAsset()
 	return GetEditingObject();
 }
 
-TFuture<bool> FSkeletalMeshEditor::HandleReimportMeshInternal(int32 SourceFileIndex /*= INDEX_NONE*/, bool bWithNewFile /*= false*/)
+TFuture<bool> FSkeletalMeshEditor::HandleReimportMeshInternal(const FReimportParameters& ReimportParameters)
 {
 	TSharedPtr<TPromise<bool>> Promise = MakeShared<TPromise<bool>>();
 
 	//The reimport will be asynchronous only if the reimport manager use Interchange.
-	UE::Interchange::FAssetImportResultRef Result = FReimportManager::Instance()->ReimportAsync(SkeletalMesh, true, true, TEXT(""), nullptr, SourceFileIndex, bWithNewFile);
+	constexpr bool bAutomatedFalse = false;
+	UE::Interchange::FAssetImportResultRef Result = FReimportManager::Instance()->ReimportAsync(SkeletalMesh, true, true, TEXT(""), nullptr, ReimportParameters.SourceFileIndex, ReimportParameters.bWithNewFile, bAutomatedFalse, ReimportParameters.bReimportWithDialog);
 
 	Result->OnDone([Promise, SkeletonTreePtr = SkeletonTree, WeakSkeletalMesh = TWeakObjectPtr<USkeletalMesh>(SkeletalMesh)](UE::Interchange::FImportResult& Result)
 		{
@@ -1312,18 +1313,11 @@ TFuture<bool> FSkeletalMeshEditor::HandleReimportMeshInternal(int32 SourceFileIn
 	return Promise->GetFuture();
 }
 
-void FSkeletalMeshEditor::HandleReimportMesh(int32 SourceFileIndex /*= INDEX_NONE*/)
+void FSkeletalMeshEditor::HandleReimportMesh(const FReimportParameters ReimportParameters)
 {
 	TSharedPtr<FScopedSuspendAlternateSkinWeightPreview> ScopedSuspendAlternateSkinnWeightPreview = MakeShared<FScopedSuspendAlternateSkinWeightPreview>(SkeletalMesh);
 	TSharedPtr<FScopedSkeletalMeshReregisterContexts> ScopedReregisterComponents = MakeShared<FScopedSkeletalMeshReregisterContexts>(SkeletalMesh);
-	HandleReimportMeshInternal(SourceFileIndex, false).Then([ScopedSuspendAlternateSkinnWeightPreview, ScopedReregisterComponents](TFuture<bool> ReimportResult) {});
-}
-
-void FSkeletalMeshEditor::HandleReimportMeshWithNewFile(int32 SourceFileIndex /*= INDEX_NONE*/)
-{
-	TSharedPtr<FScopedSuspendAlternateSkinWeightPreview> ScopedSuspendAlternateSkinnWeightPreview = MakeShared<FScopedSuspendAlternateSkinWeightPreview>(SkeletalMesh);
-	TSharedPtr<FScopedSkeletalMeshReregisterContexts> ScopedReregisterComponents = MakeShared<FScopedSkeletalMeshReregisterContexts>(SkeletalMesh);
-	HandleReimportMeshInternal(SourceFileIndex, true).Then([ScopedSuspendAlternateSkinnWeightPreview, ScopedReregisterComponents](TFuture<bool> ReimportResult) {});
+	HandleReimportMeshInternal(ReimportParameters).Then([ScopedSuspendAlternateSkinnWeightPreview, ScopedReregisterComponents](TFuture<bool> ReimportResult) {});
 }
 
 TFuture<bool> ReimportLodInChain(USkeletalMesh* SkeletalMesh
@@ -1429,7 +1423,7 @@ TFuture<bool> ReimportAllCustomLODs(USkeletalMesh* SkeletalMesh, UDebugSkelMeshC
 	return Promise->GetFuture();
 }
 
-void FSkeletalMeshEditor::HandleReimportAllMeshInternal(int32 SourceFileIndex, bool bWithNewFile)
+void FSkeletalMeshEditor::HandleReimportAllMeshInternal(const FReimportParameters& ReimportParameters)
 {
 	// Reimport the asset
 	if (SkeletalMesh)
@@ -1437,29 +1431,23 @@ void FSkeletalMeshEditor::HandleReimportAllMeshInternal(int32 SourceFileIndex, b
 		TSharedPtr<FScopedSuspendAlternateSkinWeightPreview> ScopedSuspendAlternateSkinnWeightPreview = MakeShared<FScopedSuspendAlternateSkinWeightPreview>(SkeletalMesh);
 		TSharedPtr<FScopedSkeletalMeshReregisterContexts> ScopedReregisterComponents = MakeShared<FScopedSkeletalMeshReregisterContexts>(SkeletalMesh);
 		//Reimport base LOD
-		HandleReimportMeshInternal(SourceFileIndex, bWithNewFile).Then([this, bWithNewFile, ScopedSuspendAlternateSkinnWeightPreview, ScopedReregisterComponents](TFuture<bool> Result)
+		HandleReimportMeshInternal(ReimportParameters).Then([this, ReimportParameters, ScopedSuspendAlternateSkinnWeightPreview, ScopedReregisterComponents](TFuture<bool> Result)
 		{
 			check(IsInGameThread());
 			//import all custom LODs
 			if (Result.Get() && SkeletalMesh->GetLODNum() > 1)
 			{
-				ReimportAllCustomLODs(SkeletalMesh.Get(), GetPersonaToolkit()->GetPreviewMeshComponent(), bWithNewFile);
+				ReimportAllCustomLODs(SkeletalMesh.Get(), GetPersonaToolkit()->GetPreviewMeshComponent(), ReimportParameters.bWithNewFile);
 			}
 		});
 	}
 }
 
-void FSkeletalMeshEditor::HandleReimportAllMesh(int32 SourceFileIndex /*= INDEX_NONE*/)
+void FSkeletalMeshEditor::HandleReimportAllMesh(const FReimportParameters ReimportParameters)
 {
-	constexpr bool bWithNewFile = false;
-	HandleReimportAllMeshInternal(SourceFileIndex, bWithNewFile);
+	HandleReimportAllMeshInternal(ReimportParameters);
 }
 
-void FSkeletalMeshEditor::HandleReimportAllMeshWithNewFile(int32 SourceFileIndex /*= INDEX_NONE*/)
-{
-	constexpr bool bWithNewFile = true;
-	HandleReimportAllMeshInternal(SourceFileIndex, bWithNewFile);
-}
 
 void FSkeletalMeshEditor::HandleOnPreviewSceneSettingsCustomized(IDetailLayoutBuilder& DetailBuilder)
 {
diff --git a/Engine/Source/Editor/SkeletalMeshEditor/Private/SkeletalMeshEditor.h b/Engine/Source/Editor/SkeletalMeshEditor/Private/SkeletalMeshEditor.h
--- a/Engine/Source/Editor/SkeletalMeshEditor/Private/SkeletalMeshEditor.h
+++ b/Engine/Source/Editor/SkeletalMeshEditor/Private/SkeletalMeshEditor.h
@@ -127,13 +127,24 @@ private:
 	
 	void HandleSelectionChanged(const TArrayView<TSharedPtr<ISkeletonTreeItem>>& InSelectedItems, ESelectInfo::Type InSelectInfo);
 
-	void HandleReimportMesh(int32 SourceFileIndex = INDEX_NONE);
-	void HandleReimportMeshWithNewFile(int32 SourceFileIndex = INDEX_NONE);
-	TFuture<bool> HandleReimportMeshInternal(int32 SourceFileIndex = INDEX_NONE, bool bWithNewFile = false);
+	struct FReimportParameters
+	{
+		FReimportParameters(int32 InSourceFileIndex, bool bInWithNewFile, bool bInReimportWithDialog)
+			: SourceFileIndex(InSourceFileIndex)
+			, bWithNewFile(bInWithNewFile)
+			, bReimportWithDialog(bInReimportWithDialog)
+		{}
+
+		int32 SourceFileIndex = INDEX_NONE;
+		bool bWithNewFile = false;
+		bool bReimportWithDialog = false;
+	};
+
+	void HandleReimportMesh(const FReimportParameters ReimportParameters);
+	TFuture<bool> HandleReimportMeshInternal(const FReimportParameters& ReimportParameters);
 
-	void HandleReimportAllMesh(int32 SourceFileIndex = INDEX_NONE);
-	void HandleReimportAllMeshWithNewFile(int32 SourceFileIndex = INDEX_NONE);
-	void HandleReimportAllMeshInternal(int32 SourceFileIndex, bool bWithNewFile);
+	void HandleReimportAllMesh(const FReimportParameters ReimportParameters);
+	void HandleReimportAllMeshInternal(const FReimportParameters& ReimportParameters);
 
 	void HandleOnPreviewSceneSettingsCustomized(IDetailLayoutBuilder& DetailBuilder);
 
@@ -163,7 +174,7 @@ private:
 	void OnRemoveSectionFromLodAndBelowMenuItemClicked(int32 LodIndex, int32 SectionIndex);
 	//////////////////////////////////////////////////////////////////////////
 
-	void RegisterReimportContextMenu(const FName InBaseMenuName);
+	void RegisterReimportContextMenu(const FName InBaseMenuName, bool bWithDialog);
 
 	static TSharedPtr<FSkeletalMeshEditor> GetSkeletalMeshEditor(const FToolMenuContext& InMenuContext);
 
diff --git a/Engine/Source/Editor/SkeletalMeshEditor/Private/SkeletalMeshEditorCommands.cpp b/Engine/Source/Editor/SkeletalMeshEditor/Private/SkeletalMeshEditorCommands.cpp
--- a/Engine/Source/Editor/SkeletalMeshEditor/Private/SkeletalMeshEditorCommands.cpp
+++ b/Engine/Source/Editor/SkeletalMeshEditor/Private/SkeletalMeshEditorCommands.cpp
@@ -10,10 +10,12 @@
 
 void FSkeletalMeshEditorCommands::RegisterCommands()
 {
-	UI_COMMAND(ReimportMesh, "Reimport Base Mesh", "Reimport the base mesh.", EUserInterfaceActionType::Button, FInputChord());
-	UI_COMMAND(ReimportMeshWithNewFile, "Reimport Base Mesh With New File", "Reimport the base mesh using a new source file.", EUserInterfaceActionType::Button, FInputChord());
-	UI_COMMAND(ReimportAllMesh, "Reimport Base Mesh + LODs", "Reimport the base mesh and all the custom LODs.", EUserInterfaceActionType::Button, FInputChord());
-	UI_COMMAND(ReimportAllMeshWithNewFile, "Reimport Base Mesh + LODs With New File", "Reimport the base mesh using a new source file and all the custom LODs (No new source file for LODs).", EUserInterfaceActionType::Button, FInputChord());
+	UI_COMMAND(ReimportMesh, "Reimport Mesh", "Reimport the base mesh.", EUserInterfaceActionType::Button, FInputChord());
+	UI_COMMAND(ReimportMeshWithNewFile, "Reimport Mesh With New File", "Reimport the base mesh using a new source file.", EUserInterfaceActionType::Button, FInputChord());
+	UI_COMMAND(ReimportAllMesh, "Reimport Mesh + LODs", "Reimport the base mesh and all the custom LODs.", EUserInterfaceActionType::Button, FInputChord());
+	UI_COMMAND(ReimportAllMeshWithNewFile, "Reimport Mesh + LODs With New File", "Reimport the base mesh using a new source file and all the custom LODs (No new source file for LODs).", EUserInterfaceActionType::Button, FInputChord());
+
+	UI_COMMAND(ReimportWithDialog, "Reimport Mesh With Dialog", "Reimport the base mesh with dialog.", EUserInterfaceActionType::Button, FInputChord());
 
 	UI_COMMAND(BakeMaterials, "Bake Out Materials", "Bake out Materials for given LOD(s).", EUserInterfaceActionType::Button, FInputChord());
 }
diff --git a/Engine/Source/Editor/SkeletalMeshEditor/Private/SkeletalMeshEditorCommands.h b/Engine/Source/Editor/SkeletalMeshEditor/Private/SkeletalMeshEditorCommands.h
--- a/Engine/Source/Editor/SkeletalMeshEditor/Private/SkeletalMeshEditorCommands.h
+++ b/Engine/Source/Editor/SkeletalMeshEditor/Private/SkeletalMeshEditorCommands.h
@@ -23,15 +23,15 @@ public:
 	virtual void RegisterCommands() override;
 
 public:
-
-	// reimport current mesh
+	// reimport
 	TSharedPtr<FUICommandInfo> ReimportMesh;
 	TSharedPtr<FUICommandInfo> ReimportMeshWithNewFile;
-
-	// reimport current mesh
 	TSharedPtr<FUICommandInfo> ReimportAllMesh;
 	TSharedPtr<FUICommandInfo> ReimportAllMeshWithNewFile;
 
+	// reimport with dialog
+	TSharedPtr<FUICommandInfo> ReimportWithDialog;
+	
 	// bake materials for this skeletal mesh
 	TSharedPtr<FUICommandInfo> BakeMaterials;
 };
diff --git a/Engine/Source/Editor/StaticMeshEditor/Private/StaticMeshEditor.cpp b/Engine/Source/Editor/StaticMeshEditor/Private/StaticMeshEditor.cpp
--- a/Engine/Source/Editor/StaticMeshEditor/Private/StaticMeshEditor.cpp
+++ b/Engine/Source/Editor/StaticMeshEditor/Private/StaticMeshEditor.cpp
@@ -843,13 +843,43 @@ void FStaticMeshEditor::BindCommands()
 		Commands.SaveGeneratedLODs,
 		FExecuteAction::CreateSP(this, &FStaticMeshEditor::OnSaveGeneratedLODs));
 
-	UICommandList->MapAction(
-		Commands.ReimportMesh,
-		FExecuteAction::CreateSP(this, &FStaticMeshEditor::HandleReimportMesh));
+	{
+		constexpr bool bWithNewFileTrue = true;
+		constexpr bool bWithNewFileFalse = false;
+		constexpr bool bWithDialogTrue = true;
+		constexpr bool bWithDialogFalse = false;
+		UICommandList->MapAction(
+			Commands.ReimportMesh,
+			FExecuteAction::CreateSP(this, &FStaticMeshEditor::HandleReimportMesh, FReimportParameters(bWithNewFileFalse, bWithDialogFalse)));
 
-	UICommandList->MapAction(
-		Commands.ReimportAllMesh,
-		FExecuteAction::CreateSP(this, &FStaticMeshEditor::HandleReimportAllMesh));
+		UICommandList->MapAction(
+			Commands.ReimportAllMesh,
+			FExecuteAction::CreateSP(this, &FStaticMeshEditor::HandleReimportAllMesh, FReimportParameters(bWithNewFileFalse, bWithDialogFalse)));
+
+		UICommandList->MapAction(
+			Commands.ReimportMeshWithNewFile,
+			FExecuteAction::CreateSP(this, &FStaticMeshEditor::HandleReimportMesh, FReimportParameters(bWithNewFileTrue, bWithDialogFalse)));
+
+		UICommandList->MapAction(
+			Commands.ReimportAllMeshWithNewFile,
+			FExecuteAction::CreateSP(this, &FStaticMeshEditor::HandleReimportAllMesh, FReimportParameters(bWithNewFileTrue, bWithDialogFalse)));
+
+		UICommandList->MapAction(
+			Commands.ReimportMeshWithDialog,
+			FExecuteAction::CreateSP(this, &FStaticMeshEditor::HandleReimportMesh, FReimportParameters(bWithNewFileFalse, bWithDialogTrue)));
+
+		UICommandList->MapAction(
+			Commands.ReimportAllMeshWithDialog,
+			FExecuteAction::CreateSP(this, &FStaticMeshEditor::HandleReimportAllMesh, FReimportParameters(bWithNewFileFalse, bWithDialogTrue)));
+
+		UICommandList->MapAction(
+			Commands.ReimportMeshWithNewFileWithDialog,
+			FExecuteAction::CreateSP(this, &FStaticMeshEditor::HandleReimportMesh, FReimportParameters(bWithNewFileTrue, bWithDialogTrue)));
+
+		UICommandList->MapAction(
+			Commands.ReimportAllMeshWithNewFileWithDialog,
+			FExecuteAction::CreateSP(this, &FStaticMeshEditor::HandleReimportAllMesh, FReimportParameters(bWithNewFileTrue, bWithDialogTrue)));
+	}
 
 	// Collision Menu
 	UICommandList->MapAction(
@@ -898,26 +928,56 @@ void FStaticMeshEditor::ExtendToolBar()
 			{
 				if (TSharedPtr<FStaticMeshEditor> StaticMeshEditor = StaticMeshEditor::GetStaticMeshEditorFromMenuContext(InSection))
 				{
-					auto ConstructReimportContextMenu = [](UToolMenu* InMenu)
+					//Reimport without dialog
 					{
-						FToolMenuSection& Section = InMenu->AddSection("Reimport");
-						Section.AddMenuEntry(FStaticMeshEditorCommands::Get().ReimportMesh);
-						Section.AddMenuEntry(FStaticMeshEditorCommands::Get().ReimportAllMesh);
-					};
-					
-					FToolMenuEntry& ReimportMeshEntry = InSection.AddEntry(FToolMenuEntry::InitToolBarButton(FStaticMeshEditorCommands::Get().ReimportMesh));
-					ReimportMeshEntry.StyleNameOverride = "CalloutToolbar";
-								
-					FToolMenuEntry& ReimportContextMenuEntry = InSection.AddEntry(FToolMenuEntry::InitComboButton(
-						"ReimportContextMenu",
-						FUIAction(),
-						FNewToolMenuDelegate::CreateLambda(ConstructReimportContextMenu),
-						TAttribute<FText>(),
-						TAttribute<FText>(),
-						TAttribute<FSlateIcon>(),
-						true
-					));
-					ReimportContextMenuEntry.StyleNameOverride = "CalloutToolbar";
+						auto ConstructReimportContextMenu = [](UToolMenu* InMenu)
+							{
+								FToolMenuSection& Section = InMenu->AddSection("Reimport");
+								Section.AddMenuEntry(FStaticMeshEditorCommands::Get().ReimportMesh);
+								Section.AddMenuEntry(FStaticMeshEditorCommands::Get().ReimportMeshWithNewFile);
+								Section.AddMenuEntry(FStaticMeshEditorCommands::Get().ReimportAllMesh);
+								Section.AddMenuEntry(FStaticMeshEditorCommands::Get().ReimportAllMeshWithNewFile);
+							};
+
+						FToolMenuEntry& ReimportMeshEntry = InSection.AddEntry(FToolMenuEntry::InitToolBarButton(FStaticMeshEditorCommands::Get().ReimportMesh));
+						ReimportMeshEntry.StyleNameOverride = "CalloutToolbar";
+
+						FToolMenuEntry& ReimportContextMenuEntry = InSection.AddEntry(FToolMenuEntry::InitComboButton(
+							"ReimportContextMenu",
+							FUIAction(),
+							FNewToolMenuDelegate::CreateLambda(ConstructReimportContextMenu),
+							TAttribute<FText>(),
+							TAttribute<FText>(),
+							TAttribute<FSlateIcon>(),
+							true
+						));
+						ReimportContextMenuEntry.StyleNameOverride = "CalloutToolbar";
+					}
+
+					//Reimport force dialog
+					{
+						auto ConstructReimportWithDialogContextMenu = [](UToolMenu* InMenu)
+							{
+								FToolMenuSection& Section = InMenu->AddSection("Reimport");
+								Section.AddMenuEntry(FStaticMeshEditorCommands::Get().ReimportMeshWithDialog);
+								Section.AddMenuEntry(FStaticMeshEditorCommands::Get().ReimportMeshWithNewFileWithDialog);
+								Section.AddMenuEntry(FStaticMeshEditorCommands::Get().ReimportAllMeshWithDialog);
+								Section.AddMenuEntry(FStaticMeshEditorCommands::Get().ReimportAllMeshWithNewFileWithDialog);
+							};
+						FToolMenuEntry& ReimportMeshWithDialogEntry = InSection.AddEntry(FToolMenuEntry::InitToolBarButton(FStaticMeshEditorCommands::Get().ReimportMeshWithDialog));
+						ReimportMeshWithDialogEntry.StyleNameOverride = "CalloutToolbar";
+
+						FToolMenuEntry& ReimportWithDialogContextMenuEntry = InSection.AddEntry(FToolMenuEntry::InitComboButton(
+							"ReimportWithDialogContextMenu",
+							FUIAction(),
+							FNewToolMenuDelegate::CreateLambda(ConstructReimportWithDialogContextMenu),
+							TAttribute<FText>(),
+							TAttribute<FText>(),
+							TAttribute<FSlateIcon>(),
+							true
+						));
+						ReimportWithDialogContextMenuEntry.StyleNameOverride = "CalloutToolbar";
+					}
 
 				}
 			}));
@@ -1755,22 +1815,50 @@ void FStaticMeshEditor::ComboBoxSelectionChanged( TSharedPtr<FString> NewSelecti
 	GetStaticMeshViewport()->RefreshViewport();
 }
 
-void FStaticMeshEditor::HandleReimportMesh()
+void FStaticMeshEditor::HandleReimportMesh(const FReimportParameters ReimportParameters)
 {
 	// Reimport the asset
 	if (StaticMesh)
 	{
-		FReimportManager::Instance()->ReimportAsync(StaticMesh, true);
+		constexpr bool bAskForNewFileIfMissingTrue = true;
+		constexpr bool bShowNotificationTrue = true;
+		const FString PreferredReimportFileEmpty = TEXT("");
+		constexpr FReimportHandler* SpecifiedReimportHandlerNull = nullptr;
+		constexpr int32 SourceFileIndex = INDEX_NONE;
+		constexpr bool bAutomatedFalse = false;
+		FReimportManager::Instance()->ReimportAsync(StaticMesh
+			, bAskForNewFileIfMissingTrue
+			, bShowNotificationTrue
+			, PreferredReimportFileEmpty
+			, SpecifiedReimportHandlerNull
+			, SourceFileIndex
+			, ReimportParameters.bWithNewFile
+			, bAutomatedFalse
+			, ReimportParameters.bWithDialog);
 	}
 }
 
-void FStaticMeshEditor::HandleReimportAllMesh()
+void FStaticMeshEditor::HandleReimportAllMesh(const FReimportParameters ReimportParameters)
 {
 	// Reimport the asset
 	if (StaticMesh)
 	{
+		constexpr bool bAskForNewFileIfMissingTrue = true;
+		constexpr bool bShowNotificationTrue = true;
+		const FString PreferredReimportFileEmpty = TEXT("");
+		constexpr FReimportHandler* SpecifiedReimportHandlerNull = nullptr;
+		constexpr int32 SourceFileIndex = INDEX_NONE;
+		constexpr bool bAutomatedFalse = false;
 		//Reimport base LOD, generated mesh will be rebuild here, the static mesh is always using the base mesh to reduce LOD
-		if (FReimportManager::Instance()->Reimport(StaticMesh, true))
+		if (FReimportManager::Instance()->Reimport(StaticMesh
+			, bAskForNewFileIfMissingTrue
+			, bShowNotificationTrue
+			, PreferredReimportFileEmpty
+			, SpecifiedReimportHandlerNull
+			, SourceFileIndex
+			, ReimportParameters.bWithNewFile
+			, bAutomatedFalse
+			, ReimportParameters.bWithDialog))
 		{
 			//Reimport all custom LODs
 			for (int32 LodIndex = 1; LodIndex < StaticMesh->GetNumLODs(); ++LodIndex)
diff --git a/Engine/Source/Editor/StaticMeshEditor/Private/StaticMeshEditor.h b/Engine/Source/Editor/StaticMeshEditor/Private/StaticMeshEditor.h
--- a/Engine/Source/Editor/StaticMeshEditor/Private/StaticMeshEditor.h
+++ b/Engine/Source/Editor/StaticMeshEditor/Private/StaticMeshEditor.h
@@ -363,11 +363,22 @@ private:
 	/** A general callback for the combo boxes in the Static Mesh Editor to force a viewport refresh when a selection changes. */
 	void ComboBoxSelectionChanged(TSharedPtr<FString> NewSelection, ESelectInfo::Type SelectInfo);
 
+	struct FReimportParameters
+	{
+		FReimportParameters(bool bInWithNewFile, bool bInWithDialog)
+			: bWithNewFile(bInWithNewFile)
+			, bWithDialog(bInWithDialog)
+		{}
+
+		bool bWithNewFile = false;
+		bool bWithDialog = false;
+	};
+
 	/* Callback to reimport the base mesh*/
-	void HandleReimportMesh();
+	void HandleReimportMesh(const FReimportParameters ReimportParameters);
 	
 	/* Callback to reimport the base mesh and also all custom LODs*/
-	void HandleReimportAllMesh();
+	void HandleReimportAllMesh(const FReimportParameters ReimportParameters);
 
 	/**
 	 *	Sets the editor's current mesh and refreshes various settings to correspond with the new data.
diff --git a/Engine/Source/Editor/StaticMeshEditor/Private/StaticMeshEditorActions.cpp b/Engine/Source/Editor/StaticMeshEditor/Private/StaticMeshEditorActions.cpp
--- a/Engine/Source/Editor/StaticMeshEditor/Private/StaticMeshEditorActions.cpp
+++ b/Engine/Source/Editor/StaticMeshEditor/Private/StaticMeshEditorActions.cpp
@@ -24,10 +24,14 @@ void FStaticMeshEditorCommands::RegisterCommands()
 	UI_COMMAND(SetShowSockets, "Sockets", "Displays the static mesh sockets.", EUserInterfaceActionType::ToggleButton, FInputChord(EModifierKey::Alt, EKeys::S));
 	UI_COMMAND(SetDrawAdditionalData, "Additional Data", "Draw additional user data associated with asset.", EUserInterfaceActionType::ToggleButton, FInputChord());
 
-	UI_COMMAND(ReimportMesh, "Reimport Base Mesh", "Reimport the base mesh.", EUserInterfaceActionType::Button, FInputChord());
-	UI_COMMAND(ReimportMeshWithNewFile, "Reimport Base Mesh With New File", "Reimport the base mesh using a new source file.", EUserInterfaceActionType::Button, FInputChord());
-	UI_COMMAND(ReimportAllMesh, "Reimport Base Mesh + LODs", "Reimport the base mesh and all the custom LODs.", EUserInterfaceActionType::Button, FInputChord());
-	UI_COMMAND(ReimportAllMeshWithNewFile, "Reimport Base Mesh + LODs With New File", "Reimport the base mesh using a new source file and all the custom LODs (No new source file for the custom LODs).", EUserInterfaceActionType::Button, FInputChord());
+	UI_COMMAND(ReimportMesh, "Reimport Mesh", "Reimport the base mesh.", EUserInterfaceActionType::Button, FInputChord());
+	UI_COMMAND(ReimportMeshWithNewFile, "Reimport Mesh With New File", "Reimport the base mesh using a new source file.", EUserInterfaceActionType::Button, FInputChord());
+	UI_COMMAND(ReimportAllMesh, "Reimport Mesh + LODs", "Reimport the base mesh and all the custom LODs.", EUserInterfaceActionType::Button, FInputChord());
+	UI_COMMAND(ReimportAllMeshWithNewFile, "Reimport Mesh + LODs With New File", "Reimport the base mesh using a new source file and all the custom LODs (No new source file for the custom LODs).", EUserInterfaceActionType::Button, FInputChord());
+	UI_COMMAND(ReimportMeshWithDialog, "Reimport Mesh With Dialog", "Reimport the base mesh and show the re-import dialog.", EUserInterfaceActionType::Button, FInputChord());
+	UI_COMMAND(ReimportMeshWithNewFileWithDialog, "Reimport Mesh With New File", "Reimport the base mesh using a new source file.", EUserInterfaceActionType::Button, FInputChord());
+	UI_COMMAND(ReimportAllMeshWithDialog, "Reimport Mesh + LODs", "Reimport the base mesh and all the custom LODs.", EUserInterfaceActionType::Button, FInputChord());
+	UI_COMMAND(ReimportAllMeshWithNewFileWithDialog, "Reimport Mesh + LODs With New File", "Reimport the base mesh using a new source file and all the custom LODs (No new source file for the custom LODs).", EUserInterfaceActionType::Button, FInputChord());
 
 	UI_COMMAND(ToggleShowNormals, "Normals", "Toggles display of vertex normals in the Preview Pane.", EUserInterfaceActionType::ToggleButton, FInputChord());
 	UI_COMMAND(ToggleShowTangents, "Tangents", "Toggles display of vertex tangents in the Preview Pane.", EUserInterfaceActionType::ToggleButton, FInputChord());
diff --git a/Engine/Source/Editor/StaticMeshEditor/Public/StaticMeshEditorActions.h b/Engine/Source/Editor/StaticMeshEditor/Public/StaticMeshEditorActions.h
--- a/Engine/Source/Editor/StaticMeshEditor/Public/StaticMeshEditorActions.h
+++ b/Engine/Source/Editor/StaticMeshEditor/Public/StaticMeshEditorActions.h
@@ -53,6 +53,10 @@ public:
 	TSharedPtr<FUICommandInfo> ReimportMeshWithNewFile;
 	TSharedPtr<FUICommandInfo> ReimportAllMesh;
 	TSharedPtr<FUICommandInfo> ReimportAllMeshWithNewFile;
+	TSharedPtr<FUICommandInfo> ReimportMeshWithDialog;
+	TSharedPtr<FUICommandInfo> ReimportMeshWithNewFileWithDialog;
+	TSharedPtr<FUICommandInfo> ReimportAllMeshWithDialog;
+	TSharedPtr<FUICommandInfo> ReimportAllMeshWithNewFileWithDialog;
 
 	// toolbar commands
 	TSharedPtr<FUICommandInfo> ToggleShowNormals;
diff --git a/Engine/Source/Editor/UnrealEd/Private/Editor.cpp b/Engine/Source/Editor/UnrealEd/Private/Editor.cpp
--- a/Engine/Source/Editor/UnrealEd/Private/Editor.cpp
+++ b/Engine/Source/Editor/UnrealEd/Private/Editor.cpp
@@ -279,9 +279,9 @@ void FReimportManager::UpdateReimportPath(UObject* Obj, const FString& Filename,
 }
 
 
-bool FReimportManager::Reimport(UObject* Obj, bool bAskForNewFileIfMissing, bool bShowNotification, FString PreferredReimportFile, FReimportHandler* SpecifiedReimportHandler, int32 SourceFileIndex, bool bForceNewFile /*= false*/, bool bAutomated /*= false*/)
+bool FReimportManager::Reimport(UObject* Obj, bool bAskForNewFileIfMissing, bool bShowNotification, FString PreferredReimportFile, FReimportHandler* SpecifiedReimportHandler, int32 SourceFileIndex, bool bForceNewFile /*= false*/, bool bAutomated /*= false*/, bool bInForceShowDialog /*= false*/)
 {
-	UE::Interchange::FAssetImportResultRef ImportResult = ReimportAsync(Obj, bAskForNewFileIfMissing, bShowNotification, PreferredReimportFile, SpecifiedReimportHandler, SourceFileIndex, bForceNewFile, bAutomated);
+	UE::Interchange::FAssetImportResultRef ImportResult = ReimportAsync(Obj, bAskForNewFileIfMissing, bShowNotification, PreferredReimportFile, SpecifiedReimportHandler, SourceFileIndex, bForceNewFile, bAutomated, bInForceShowDialog);
 	ImportResult->WaitUntilDone();
 	const TArray<UInterchangeResult*>& Results = ImportResult->GetResults()->GetResults();
 	for (const UInterchangeResult* InterchangeResult : Results)
@@ -294,7 +294,7 @@ bool FReimportManager::Reimport(UObject* Obj, bool bAskForNewFileIfMissing, bool
 	return true;
 }
 
-UE::Interchange::FAssetImportResultRef FReimportManager::ReimportAsync(UObject* Obj, bool bAskForNewFileIfMissing, bool bShowNotification, FString PreferredReimportFile, FReimportHandler* SpecifiedReimportHandler, int32 SourceFileIndex, bool bForceNewFile /*= false*/, bool bAutomated /*= false*/)
+UE::Interchange::FAssetImportResultRef FReimportManager::ReimportAsync(UObject* Obj, bool bAskForNewFileIfMissing, bool bShowNotification, FString PreferredReimportFile, FReimportHandler* SpecifiedReimportHandler, int32 SourceFileIndex, bool bForceNewFile /*= false*/, bool bAutomated /*= false*/, bool bInForceShowDialog /*= false*/)
 {
 	UE::Interchange::FAssetImportResultRef ImportResultSynchronous = MakeShared< UE::Interchange::FImportResult, ESPMode::ThreadSafe >();
 	// Warn that were about to reimport, so prep for it
@@ -470,6 +470,7 @@ UE::Interchange::FAssetImportResultRef FReimportManager::ReimportAsync(UObject*
 							ImportAssetParameters.bIsAutomated = GIsAutomationTesting || FApp::IsUnattended() || IsRunningCommandlet() || GIsRunningUnattendedScript;
 							ImportAssetParameters.ReimportAsset = Obj;
 							ImportAssetParameters.ReimportSourceIndex = SourceFileIndex;
+							ImportAssetParameters.bForceShowDialog = bInForceShowDialog;
 							UE::Interchange::FAssetImportResultRef ImportResult = InterchangeManager.ImportAssetAsync(FString(), ScopedSourceData.GetSourceData(), ImportAssetParameters);
 
 							TFunction<void(UE::Interchange::FImportResult&)> AppendAndBroadcastImportResultIfNeeded =
@@ -492,7 +493,13 @@ UE::Interchange::FAssetImportResultRef FReimportManager::ReimportAsync(UObject*
 				// Do the reimport
 				const bool bOriginalAutomated = CanReimportHandler->IsAutomatedReimport();
 				CanReimportHandler->SetAutomatedReimport(bAutomated);
+				const bool bOriginalForceShowDialog = CanReimportHandler->IsForceShowDialog();
+				CanReimportHandler->SetForceShowDialog(bInForceShowDialog);
+				
+				//Reimport the asset
 				EReimportResult::Type Result = CanReimportHandler->Reimport( Obj, SourceFileIndex );
+				
+				CanReimportHandler->SetForceShowDialog(bOriginalForceShowDialog);
 				CanReimportHandler->SetAutomatedReimport(bOriginalAutomated);
 				// Even if the reimport has been successful, check that the originating object is still valid
 				// The reimport might be a reimport to level which triggered the deletion of the object
diff --git a/Engine/Source/Editor/UnrealEd/Private/Factories/EditorFactories.cpp b/Engine/Source/Editor/UnrealEd/Private/Factories/EditorFactories.cpp
--- a/Engine/Source/Editor/UnrealEd/Private/Factories/EditorFactories.cpp
+++ b/Engine/Source/Editor/UnrealEd/Private/Factories/EditorFactories.cpp
@@ -5850,7 +5850,7 @@ EReimportResult::Type UReimportFbxStaticMeshFactory::Reimport( UObject* Obj )
 	}
 	//Prevent any UI for automation, unattended and commandlet
 	const bool IsUnattended = IsAutomatedImport() || FApp::IsUnattended() || IsRunningCommandlet() || GIsRunningUnattendedScript;
-	const bool ShowImportDialogAtReimport = GetDefault<UEditorPerProjectUserSettings>()->bShowImportDialogAtReimport && !IsUnattended;
+	const bool ShowImportDialogAtReimport = (GetDefault<UEditorPerProjectUserSettings>()->bShowImportDialogAtReimport || bForceShowDialog) && !IsUnattended;
 
 	if (ImportData == nullptr)
 	{
@@ -6229,7 +6229,7 @@ EReimportResult::Type UReimportFbxSkeletalMeshFactory::Reimport( UObject* Obj, i
 	bool bSuccess = false;
 	//Prevent any UI for automation, unattended and commandlet
 	const bool IsUnattended = IsAutomatedImport() || FApp::IsUnattended() || IsRunningCommandlet() || GIsRunningUnattendedScript;
-	const bool ShowImportDialogAtReimport = GetDefault<UEditorPerProjectUserSettings>()->bShowImportDialogAtReimport && !IsUnattended;
+	const bool ShowImportDialogAtReimport = (GetDefault<UEditorPerProjectUserSettings>()->bShowImportDialogAtReimport || bForceShowDialog) && !IsUnattended;
 
 	if (ImportData == nullptr)
 	{
@@ -6712,7 +6712,7 @@ EReimportResult::Type UReimportFbxAnimSequenceFactory::Reimport( UObject* Obj )
 	}
 
 	bool bOutImportAll = false;
-	if ( UEditorEngine::ReimportFbxAnimation(Skeleton, AnimSequence, ImportData, *Filename, bOutImportAll, bShowOption && !IsAutomatedImport(), OverrideImportUI) )
+	if ( UEditorEngine::ReimportFbxAnimation(Skeleton, AnimSequence, ImportData, *Filename, bOutImportAll, (bShowOption || bForceShowDialog) && !IsAutomatedImport(), OverrideImportUI) )
 	{
 		if (bOutImportAll)
 		{
diff --git a/Engine/Source/Editor/UnrealEd/Public/EditorReimportHandler.h b/Engine/Source/Editor/UnrealEd/Public/EditorReimportHandler.h
--- a/Engine/Source/Editor/UnrealEd/Public/EditorReimportHandler.h
+++ b/Engine/Source/Editor/UnrealEd/Public/EditorReimportHandler.h
@@ -63,7 +63,7 @@ public:
 	 *
 	 * @return	true if the object was handled by one of the reimport handlers; false otherwise
 	 */
-	UNREALED_API virtual bool Reimport( UObject* Obj, bool bAskForNewFileIfMissing = false, bool bShowNotification = true, FString PreferredReimportFile = TEXT(""), FReimportHandler* SpecifiedReimportHandler = nullptr, int32 SourceFileIndex = INDEX_NONE, bool bForceNewFile = false, bool bAutomated = false);
+	UNREALED_API virtual bool Reimport( UObject* Obj, bool bAskForNewFileIfMissing = false, bool bShowNotification = true, FString PreferredReimportFile = TEXT(""), FReimportHandler* SpecifiedReimportHandler = nullptr, int32 SourceFileIndex = INDEX_NONE, bool bForceNewFile = false, bool bAutomated = false, bool bInForceShowDialog = false);
 
 	/**
 	 * Attempt to reimport the specified object from its source by giving registered reimport
@@ -78,7 +78,7 @@ public:
 	 *
 	 * @return	FAssetImportResultRef
 	 */
-	UNREALED_API virtual UE::Interchange::FAssetImportResultRef ReimportAsync(UObject* Obj, bool bAskForNewFileIfMissing = false, bool bShowNotification = true, FString PreferredReimportFile = TEXT(""), FReimportHandler* SpecifiedReimportHandler = nullptr, int32 SourceFileIndex = INDEX_NONE, bool bForceNewFile = false, bool bAutomated = false);
+	UNREALED_API virtual UE::Interchange::FAssetImportResultRef ReimportAsync(UObject* Obj, bool bAskForNewFileIfMissing = false, bool bShowNotification = true, FString PreferredReimportFile = TEXT(""), FReimportHandler* SpecifiedReimportHandler = nullptr, int32 SourceFileIndex = INDEX_NONE, bool bInForceNewFile = false, bool bAutomated = false, bool bForceShowDialog = false);
 
 	/**
 	 * Attemp to reimport all specified objects. This function will verify that all source file exist and ask the user
@@ -316,7 +316,20 @@ public:
 		return false;
 	}
 
+	/** Sets force show dialog, when True dialog will prompts even if the show dialog setting is false. */
+	virtual void SetForceShowDialog(const bool InbForceShowDialog)
+	{
+		bForceShowDialog = InbForceShowDialog;
+	}
+
+	/** When True dialog prompts will prompts even if the show dialog setting is false. */
+	virtual bool IsForceShowDialog() const
+	{
+		return bForceShowDialog;
+	}
+
 protected:
 	FString PreferredReimportPath;
 	bool bAutomatedReimport;
+	bool bForceShowDialog = false;
 };
diff --git a/Engine/Source/Runtime/Interchange/Engine/Private/InterchangeManager.cpp b/Engine/Source/Runtime/Interchange/Engine/Private/InterchangeManager.cpp
--- a/Engine/Source/Runtime/Interchange/Engine/Private/InterchangeManager.cpp
+++ b/Engine/Source/Runtime/Interchange/Engine/Private/InterchangeManager.cpp
@@ -1483,7 +1483,7 @@ void UInterchangeManager::StartQueuedTasks(bool bCancelAllTasks /*= false*/)
 				for (int32 SourceDataIndex = 0; SourceDataIndex < QueuedTaskData.AsyncHelper->SourceDatas.Num(); ++SourceDataIndex)
 				{
 					//Log the source we begin importing
-					UE_LOG(LogInterchangeEngine, Display, TEXT("Interchange start importing source [%s]"), *QueuedTaskData.AsyncHelper->SourceDatas[SourceDataIndex]->ToDisplayString());
+					UE_LOG(LogInterchangeEngine, Display, TEXT("Interchange start importing source [%s]"), *QueuedTaskData.AsyncHelper->SourceDatas[SourceDataIndex]->GetFilename());
 					int32 TranslatorTaskIndex = QueuedTaskData.AsyncHelper->TranslatorTasks.Add(TGraphTask<UE::Interchange::FTaskTranslator>::CreateTask().ConstructAndDispatchWhenReady(SourceDataIndex, WeakAsyncHelper));
 					PipelinePrerequistes.Add(QueuedTaskData.AsyncHelper->TranslatorTasks[TranslatorTaskIndex]);
 				}
@@ -1885,7 +1885,7 @@ UInterchangeManager::ImportInternal(const FString& ContentPath, const UInterchan
 	{
 #if WITH_EDITORONLY_DATA
 		bool bShowPipelineStacksConfigurationDialog = !bIsUnattended
-			&& FInterchangeProjectSettingsUtils::ShouldShowPipelineStacksConfigurationDialog(bImportScene, bIsReimport, *SourceData)
+			&& (FInterchangeProjectSettingsUtils::ShouldShowPipelineStacksConfigurationDialog(bImportScene, bIsReimport, *SourceData) || ImportAssetParameters.bForceShowDialog)
 			&& !bImportCanceled
 			&& !IsRunningCommandlet();
 #else
@@ -1927,7 +1927,7 @@ UInterchangeManager::ImportInternal(const FString& ContentPath, const UInterchan
 				for (int32 SourceDataIndex = 0; SourceDataIndex < AsyncHelper->SourceDatas.Num(); ++SourceDataIndex)
 				{
 					//Log the source we begin importing
-					UE_LOG(LogInterchangeEngine, Display, TEXT("Interchange start importing source [%s]"), *AsyncHelper->SourceDatas[SourceDataIndex]->ToDisplayString());
+					UE_LOG(LogInterchangeEngine, Display, TEXT("Interchange start importing source [%s]"), *AsyncHelper->SourceDatas[SourceDataIndex]->GetFilename());
 					int32 TranslatorTaskIndex = AsyncHelper->TranslatorTasks.Add(TGraphTask<UE::Interchange::FTaskTranslator>::CreateTask(nullptr, ENamedThreads::GameThread_Local).ConstructAndDispatchWhenReady(SourceDataIndex, AsyncHelper));
 					AsyncHelper->TranslatorTasks[TranslatorTaskIndex]->Wait();
 				}
diff --git a/Engine/Source/Runtime/Interchange/Engine/Public/InterchangeManager.h b/Engine/Source/Runtime/Interchange/Engine/Public/InterchangeManager.h
--- a/Engine/Source/Runtime/Interchange/Engine/Public/InterchangeManager.h
+++ b/Engine/Source/Runtime/Interchange/Engine/Public/InterchangeManager.h
@@ -386,6 +386,10 @@ struct FImportAssetParameters
 	UPROPERTY(EditAnywhere, BlueprintReadWrite, Category = "Interchange|ImportAsset")
 	bool bReplaceExisting = true;
 
+	/** If true this import must show the import dialog and ignore the show dialog settings. */
+	UPROPERTY(EditAnywhere, BlueprintReadWrite, Category = "Interchange|ImportAsset")
+	bool bForceShowDialog = false;
+
 	/* Delegates used to track the imported objects. */
 
 	// This is called each time an asset is imported or reimported from the import call.
