- Add a morph target information saying its generated by engine, those morph target are not erase by a re-import.
- Add skeletal mesh subsystem functions API to allow user to find and set some morph target to generated by engine.
- Fix iterator on mutable container in skeletal mesh factory
- Make sure we do not recycle morph target sub object out side of the game thread.

#rnx

diff --git a/Engine/Plugins/Interchange/Runtime/Source/Import/Private/Mesh/InterchangeSkeletalMeshFactory.cpp b/Engine/Plugins/Interchange/Runtime/Source/Import/Private/Mesh/InterchangeSkeletalMeshFactory.cpp
--- a/Engine/Plugins/Interchange/Runtime/Source/Import/Private/Mesh/InterchangeSkeletalMeshFactory.cpp
+++ b/Engine/Plugins/Interchange/Runtime/Source/Import/Private/Mesh/InterchangeSkeletalMeshFactory.cpp
@@ -2301,8 +2301,35 @@ void UInterchangeSkeletalMeshFactory::SetupObject_GameThread(const FSetupObjectP
 			TSharedPtr<FInterchangeSkeletalMeshAlternateSkinWeightPostImportTask> SkeletalMeshPostImportTask = nullptr;
 			TArray<FSkinWeightProfileInfo>& SkinProfiles = SkeletalMesh->GetSkinWeightProfiles();
 			SkinProfiles = ImportAssetObjectData.ExistingSkinWeightProfileInfos;
+			//Since SkeletalMesh->SaveLODImportedData can remove invalid profile we can't use an iterator or a for loop
+			//We store all the name in a TArray and we will search for them
+			TArray<FName> ProfileList;
 			for (const FSkinWeightProfileInfo& ProfileInfo : SkinProfiles)
 			{
+				ProfileList.AddUnique(ProfileInfo.Name);
+			}
+			FSkinWeightProfileInfo EmptyProfile = FSkinWeightProfileInfo();
+			EmptyProfile.Name = NAME_None;
+			//Return a profile reference, if it doesn't found any profile, it will return the empty profile
+			auto GetProfileFromName = [&SkinProfiles, &EmptyProfile](const FName ProfileNameToSearch)
+				{
+					for (const FSkinWeightProfileInfo& ProfileInfo : SkinProfiles)
+					{
+						if (ProfileInfo.Name == ProfileNameToSearch)
+						{
+							return ProfileInfo;
+						}
+					}
+					return EmptyProfile;
+				};
+			
+			for(const FName& ProfileName : ProfileList)
+			{
+				const FSkinWeightProfileInfo& ProfileInfo = GetProfileFromName(ProfileName);
+				if (ProfileInfo.Name == NAME_None)
+				{
+					continue;
+				}
 				for (int32 LodIndex = 0; LodIndex < LodCount; ++LodIndex)
 				{
 					if (!SkeletalMesh->IsValidLODIndex(LodIndex)
diff --git a/Engine/Plugins/Runtime/GeometryScripting/Source/GeometryScriptingCore/Private/MeshAssetFunctions.cpp b/Engine/Plugins/Runtime/GeometryScripting/Source/GeometryScriptingCore/Private/MeshAssetFunctions.cpp
--- a/Engine/Plugins/Runtime/GeometryScripting/Source/GeometryScriptingCore/Private/MeshAssetFunctions.cpp
+++ b/Engine/Plugins/Runtime/GeometryScripting/Source/GeometryScriptingCore/Private/MeshAssetFunctions.cpp
@@ -1239,6 +1239,14 @@ UDynamicMesh* UGeometryScriptLibrary_StaticMeshFunctions::CopyMorphTargetToSkele
 		return FromDynamicMesh;
 	}
 
+#if WITH_EDITORONLY_DATA
+	if (FSkeletalMeshLODInfo* LodInfo = ToSkeletalMeshAsset->GetLODInfo(TargetLOD.LODIndex))
+	{
+		constexpr bool bGeneratedByEngineTrue = true;
+		LodInfo->ImportedMorphTargetSourceFilename.FindOrAdd(MorphTargetName.ToString()).SetGeneratedByEngine(bGeneratedByEngineTrue);
+	}
+#endif //WITH_EDITORONLY_DATA
+
 	ToSkeletalMeshAsset->CommitMeshDescription(TargetLOD.LODIndex);
 
 	if (Options.bDeferMeshPostEditChange == false)
diff --git a/Engine/Source/Developer/SkeletalMeshUtilitiesCommon/Private/LODUtilities.cpp b/Engine/Source/Developer/SkeletalMeshUtilitiesCommon/Private/LODUtilities.cpp
--- a/Engine/Source/Developer/SkeletalMeshUtilitiesCommon/Private/LODUtilities.cpp
+++ b/Engine/Source/Developer/SkeletalMeshUtilitiesCommon/Private/LODUtilities.cpp
@@ -4284,7 +4284,8 @@ void FLODUtilities::BackupCustomImportedMorphTargetData(USkeletalMesh* SkeletalM
 
 bool FLODUtilities::RestoreCustomImportedMorphTargetData(USkeletalMesh* SkeletalMesh, const int32 LodIndex, FMeshDescription& LodMeshDescription, const TMap<FString, TArray<FMorphTargetLodBackupData>>& BackupImportedMorphTargetData)
 {
-	auto RemoveInvalidMorphTargetLOD = [SkeletalMesh, LodIndex](const FString& MorphTargetNameStr)
+	TMap<FString, FMorphTargetImportedSourceFileInfo> RemovedImportedMorphTargetInfo;
+	auto RemoveInvalidMorphTargetLOD = [SkeletalMesh, LodIndex, &RemovedImportedMorphTargetInfo](const FString& MorphTargetNameStr)
 		{
 			FName MorphTargetName(*MorphTargetNameStr);
 			if(UMorphTarget* MorphTarget = SkeletalMesh->FindMorphTarget(MorphTargetName))
@@ -4315,11 +4316,13 @@ bool FLODUtilities::RestoreCustomImportedMorphTargetData(USkeletalMesh* Skeletal
 			{
 				if (LodInfo->ImportedMorphTargetSourceFilename.Contains(MorphTargetNameStr))
 				{
+					FMorphTargetImportedSourceFileInfo& RemoveBackup = RemovedImportedMorphTargetInfo.FindOrAdd(MorphTargetNameStr);
+					RemoveBackup = LodInfo->ImportedMorphTargetSourceFilename.FindChecked(MorphTargetNameStr);
 					LodInfo->ImportedMorphTargetSourceFilename.Remove(MorphTargetNameStr);
 				}
 			}
 		};
-
+	
 	bool bIsMeshDescriptionModified = false;
 	for (const TPair<FString, TArray<FMorphTargetLodBackupData>>& ImportedMorphTargetPair : BackupImportedMorphTargetData)
 	{
@@ -4374,6 +4377,25 @@ bool FLODUtilities::RestoreCustomImportedMorphTargetData(USkeletalMesh* Skeletal
 			}
 		}
 	}
+	if (!GIsRunningUnattendedScript && !FApp::IsUnattended() && !RemovedImportedMorphTargetInfo.IsEmpty())
+	{
+		FTextBuilder ErrorMessage;
+		ErrorMessage.AppendLine(LOCTEXT("BadTopologyMorphRemoveMessage", "Morph target(s) removed because of topology changes:"));
+		ErrorMessage.Indent();
+		for (const TPair<FString, FMorphTargetImportedSourceFileInfo>& RemovedMorphTargetInfo : RemovedImportedMorphTargetInfo)
+		{
+			if (RemovedMorphTargetInfo.Value.IsGeneratedByEngine())
+			{
+				ErrorMessage.AppendLine(FText::Format(LOCTEXT("BadTopologyMorphRemoveMessagePerMorphGenerated", "{0} [LOD {1}] Generate by Engine tools"), FText::FromString(RemovedMorphTargetInfo.Key), FText::AsNumber(LodIndex)));
+			}
+			else
+			{
+				ErrorMessage.AppendLine(FText::Format(LOCTEXT("BadTopologyMorphRemoveMessagePerMorphCustomImport", "{0} [LOD {1}] SourceFile:{2}"), FText::FromString(RemovedMorphTargetInfo.Key), FText::AsNumber(LodIndex), FText::FromString(RemovedMorphTargetInfo.Value.GetSourceFilename())));
+			}
+		}
+		ErrorMessage.Unindent();
+		FMessageDialog::Open(EAppMsgType::Ok, ErrorMessage.ToText());
+	}
 
 	return bIsMeshDescriptionModified;
 }
diff --git a/Engine/Source/Editor/Persona/Private/SRenameMorphTargetDialog.cpp b/Engine/Source/Editor/Persona/Private/SRenameMorphTargetDialog.cpp
--- a/Engine/Source/Editor/Persona/Private/SRenameMorphTargetDialog.cpp
+++ b/Engine/Source/Editor/Persona/Private/SRenameMorphTargetDialog.cpp
@@ -178,9 +178,12 @@ void SRenameMorphTargetDialog::RenameAndClose()
 								if (FMorphTargetImportedSourceFileInfo* MorphTargetImportedSourceFileInfo = LodInfo->ImportedMorphTargetSourceFilename.Find(MorphTargetToRename))
 								{
 									const FString OriginalSourceFile = MorphTargetImportedSourceFileInfo->GetSourceFilename();
+									const bool bIsGeneratedByEngine = MorphTargetImportedSourceFileInfo->IsGeneratedByEngine();
 									LodInfo->ImportedMorphTargetSourceFilename.Remove(MorphTargetToRename);
 									//Add the mapping to the lod info so we can re-import the morph target
-									LodInfo->ImportedMorphTargetSourceFilename.FindOrAdd(NewMorphTargetName).SetSourceFilename(OriginalSourceFile);
+									FMorphTargetImportedSourceFileInfo& NewData = LodInfo->ImportedMorphTargetSourceFilename.FindOrAdd(NewMorphTargetName);
+									NewData.SetSourceFilename(OriginalSourceFile);
+									NewData.SetGeneratedByEngine(bIsGeneratedByEngine);
 								}
 							}
 						}
diff --git a/Engine/Source/Editor/SkeletalMeshEditor/Private/SkeletalMeshEditorSubsystem.cpp b/Engine/Source/Editor/SkeletalMeshEditor/Private/SkeletalMeshEditorSubsystem.cpp
--- a/Engine/Source/Editor/SkeletalMeshEditor/Private/SkeletalMeshEditorSubsystem.cpp
+++ b/Engine/Source/Editor/SkeletalMeshEditor/Private/SkeletalMeshEditorSubsystem.cpp
@@ -14,6 +14,7 @@
 #include "Engine/SkinnedAssetCommon.h"
 #include "FbxMeshUtils.h"
 #include "LODUtilities.h"
+#include "Misc/ScopedSlowTask.h"
 #include "ObjectTools.h"
 #include "PhysicsAssetUtils.h"
 #include "PhysicsEngine/PhysicsAsset.h"
@@ -982,4 +983,108 @@ bool USkeletalMeshEditorSubsystem::AssignPhysicsAsset(USkeletalMesh* TargetMesh,
 	return true;
 }
 
+bool USkeletalMeshEditorSubsystem::SetMorphTargetsToGeneratedByEngine(USkeletalMesh* SkeletalMesh, const TArray<FString>& OptionalNames)
+{
+	TArray<FString> MorphTargetNames = SkeletalMesh->K2_GetAllMorphTargetNames();
+	TArray<FString> ToSetGeneratedByEngineMorphTargetNames;
+	if (OptionalNames.IsEmpty())
+	{
+		ToSetGeneratedByEngineMorphTargetNames = MorphTargetNames;
+	}
+	else
+	{
+		for (const FString& OptionalMorphTargetName : OptionalNames)
+		{
+			if (MorphTargetNames.Contains(OptionalMorphTargetName))
+			{
+				ToSetGeneratedByEngineMorphTargetNames.Add(OptionalMorphTargetName);
+			}
+		}
+
+	}
+	if (ToSetGeneratedByEngineMorphTargetNames.IsEmpty())
+	{
+		return false;
+	}
+	bool bPreEditChangeCalled = false;
+	auto DoPreEditChange = [&SkeletalMesh, &bPreEditChangeCalled]()
+		{
+			if (!bPreEditChangeCalled)
+			{
+				bPreEditChangeCalled = true;
+				SkeletalMesh->PreEditChange(nullptr);
+			}
+		};
+	const int32 LodCount = SkeletalMesh->GetLODNum();
+	for (int32 LodIndex = 0; LodIndex < LodCount; ++LodIndex)
+	{
+		FSkeletalMeshLODInfo* LodInfo = SkeletalMesh->GetLODInfo(LodIndex);
+		if (LodInfo)
+		{
+			for (const FString& MorphTargetNameStr : ToSetGeneratedByEngineMorphTargetNames)
+			{
+				constexpr bool bGeneratedByEngineTrue = true;
+				if (LodInfo->ImportedMorphTargetSourceFilename.Contains(MorphTargetNameStr))
+				{
+					FMorphTargetImportedSourceFileInfo& MorphTargetImportedSourceFileInfo = LodInfo->ImportedMorphTargetSourceFilename.FindChecked(MorphTargetNameStr);
+					if(!MorphTargetImportedSourceFileInfo.IsGeneratedByEngine())
+					{
+						DoPreEditChange();
+						MorphTargetImportedSourceFileInfo.SetGeneratedByEngine(bGeneratedByEngineTrue);
+					}
+				}
+				else
+				{
+					DoPreEditChange();
+					LodInfo->ImportedMorphTargetSourceFilename.FindOrAdd(MorphTargetNameStr).SetGeneratedByEngine(bGeneratedByEngineTrue);
+				}
+			}
+		}
+	}
+	if (bPreEditChangeCalled)
+	{
+		SkeletalMesh->PostEditChange();
+	}
+	return bPreEditChangeCalled;
+}
+
+void USkeletalMeshEditorSubsystem::SetMorphTargetsToGeneratedByEngineForAllSkeletalMesh(const TArray<FString>& OptionalNames, const TArray<FName>& OptionalPaths)
+{
+	FAssetRegistryModule& AssetRegistryModule = FModuleManager::LoadModuleChecked<FAssetRegistryModule>("AssetRegistry");
+	TArray<FAssetData> AssetDatas;
+	FARFilter Filter;
+	
+	Filter.ClassPaths.Add(USkeletalMesh::StaticClass()->GetClassPathName());
+	Filter.bRecursiveClasses = false;
+	if (!OptionalPaths.IsEmpty())
+	{
+		Filter.bRecursivePaths = true;
+		Filter.PackagePaths = OptionalPaths;
+		Filter.bIncludeOnlyOnDiskAssets = false;
+	}
+	AssetRegistryModule.Get().GetAssets(Filter, AssetDatas);
+
+	const int32 AssetCount = AssetDatas.Num();
+	FText ProgressText = FText::Format(LOCTEXT("SetMorphTargetsToGeneratedByEngineForAllSkeletalMesh_SlowTask", "Iterating {0} skeletalmeshes..."), FText::AsNumber(AssetCount));
+	FScopedSlowTask Progress(AssetCount, ProgressText);
+	Progress.MakeDialog(true);
+	int32 ModifiedSkeletalMeshCount = 0;
+	for (const FAssetData& SkeletalMeshAsset : AssetDatas)
+	{
+		//Break if user ask to cancel the operation
+		if (Progress.ShouldCancel())
+		{
+			break;
+		}
+		//SkeletalMeshAsset.FindTag()
+		USkeletalMesh* SkeletalMesh = Cast<USkeletalMesh>(SkeletalMeshAsset.GetAsset());
+		if (SetMorphTargetsToGeneratedByEngine(SkeletalMesh, OptionalNames))
+		{
+			ModifiedSkeletalMeshCount++;
+			ProgressText = FText::Format(LOCTEXT("SetMorphTargetsToGeneratedByEngineForAllSkeletalMesh_SlowTask_modified", "Iterating {0} skeletalmeshes... modified count {1}"), FText::AsNumber(AssetCount), FText::AsNumber(ModifiedSkeletalMeshCount));
+		}
+		Progress.EnterProgressFrame(1, ProgressText);
+	}
+}
+
 #undef LOCTEXT_NAMESPACE
diff --git a/Engine/Source/Editor/SkeletalMeshEditor/Public/SkeletalMeshEditorSubsystem.h b/Engine/Source/Editor/SkeletalMeshEditor/Public/SkeletalMeshEditorSubsystem.h
--- a/Engine/Source/Editor/SkeletalMeshEditor/Public/SkeletalMeshEditorSubsystem.h
+++ b/Engine/Source/Editor/SkeletalMeshEditor/Public/SkeletalMeshEditorSubsystem.h
@@ -265,4 +265,23 @@ public:
 	 */
 	UFUNCTION(BlueprintCallable, Category = "Editor Scripting | SkeletalMesh")
 	static bool AssignPhysicsAsset(USkeletalMesh* TargetMesh, UPhysicsAsset* PhysicsAsset);
+
+	/**
+	 * Set morph targets has generated by engine.
+	 * 
+	 * @param TargetMesh The mesh to convert morph targets to generated by engine.
+	 * @param OptionalNames The names of morph targets to convert. If the array is empty all morph will be converted.
+	 * @return Return true if the target mesh was modified, false otherwise.
+	 */
+	UFUNCTION(BlueprintCallable, Category = "Editor Scripting | SkeletalMesh")
+	static bool SetMorphTargetsToGeneratedByEngine(USkeletalMesh* TargetMesh, const TArray<FString>& OptionalNames);
+
+	/**
+	 * Set morph targets has generated by engine for all skeletalmesh in the content.
+	 *
+	 * @param OptionalNames The names of morph targets to convert. If the array is empty all morph will be converted.
+	 * @param OptionalPaths The paths list use to filter the research of the skeletal meshes. If the array is empty we will find all skeletal meshes.
+	 */
+	UFUNCTION(BlueprintCallable, Category = "Editor Scripting | SkeletalMesh")
+	static void SetMorphTargetsToGeneratedByEngineForAllSkeletalMesh(const TArray<FString>& OptionalNames, const TArray<FName>& OptionalPaths);
 };
diff --git a/Engine/Source/Runtime/Engine/Classes/Engine/SkinnedAssetCommon.h b/Engine/Source/Runtime/Engine/Classes/Engine/SkinnedAssetCommon.h
--- a/Engine/Source/Runtime/Engine/Classes/Engine/SkinnedAssetCommon.h
+++ b/Engine/Source/Runtime/Engine/Classes/Engine/SkinnedAssetCommon.h
@@ -105,9 +105,14 @@ struct FMorphTargetImportedSourceFileInfo
 
 	ENGINE_API void SetSourceFilename(const FString& Filename);
 
+	ENGINE_API bool IsGeneratedByEngine() const;
+
+	ENGINE_API void SetGeneratedByEngine(bool bInGeneratedByEngine);
+
 private:
 	UPROPERTY(EditAnywhere, Category = MorphTargetInfo)
 	FString SourceFilename;
+
 	
 	UPROPERTY()
 	FGuid DerivedDataHash;
diff --git a/Engine/Source/Runtime/Engine/Private/SkeletalMesh.cpp b/Engine/Source/Runtime/Engine/Private/SkeletalMesh.cpp
--- a/Engine/Source/Runtime/Engine/Private/SkeletalMesh.cpp
+++ b/Engine/Source/Runtime/Engine/Private/SkeletalMesh.cpp
@@ -2637,12 +2637,35 @@ bool USkeletalMesh::CommitMeshDescription(
 			{
 				if (!ExistingMorphTargetNames.Contains(MorphTargetName))
 				{
-					UMorphTarget* MorphTarget = NewObject<UMorphTarget>(this, MorphTargetName);
-					MorphTarget->BaseSkelMesh = this;
-					MorphTarget->ClearInternalFlags(EInternalObjectFlags::Async);
-					
-					ExistingMorphTargets.Add(MorphTarget);
-					
+					UMorphTarget* MorphTarget = nullptr;
+					// See if object already exists.
+					UObject* Obj = StaticFindObjectFastInternal(nullptr, this, MorphTargetName, true);
+					if (Obj)
+					{
+						MorphTarget = Cast<UMorphTarget>(Obj);
+						if (MorphTarget)
+						{
+							//Make sure the UObject is not garbage anymore and there is no data
+							MorphTarget->ClearGarbage();
+							MorphTarget->EmptyMorphLODModels();
+							MorphTarget->SetFlags(RF_Standalone);
+						}
+						else
+						{
+							UE_LOG(LogSkeletalMesh, Error, TEXT("Skeletal Mesh (%s) : Commit mesh description, cannot create a morph target name %s because a subobject of class %s already exist"), *GetPathName(), *(MorphTargetName.ToString()), *(Obj->GetClass()->GetName()));
+						}
+					}
+					else
+					{
+						MorphTarget = NewObject<UMorphTarget>(this, MorphTargetName);
+					}
+
+					if (MorphTarget)
+					{
+						MorphTarget->BaseSkelMesh = this;
+						MorphTarget->ClearInternalFlags(EInternalObjectFlags::Async);
+						ExistingMorphTargets.Add(MorphTarget);
+					}
 					bMorphTargetsChanged = true;
 				}
 			}
@@ -4533,7 +4556,10 @@ bool USkeletalMesh::RenameMorphTarget(FName InOldName, FName InNewName)
 		if (const FMorphTargetImportedSourceFileInfo* MorphTargetImportedSourceFileInfo = LODInfoEntry.ImportedMorphTargetSourceFilename.Find(InOldName.ToString()))
 		{
 			const FString OldFilename = MorphTargetImportedSourceFileInfo->GetSourceFilename();
-			LODInfoEntry.ImportedMorphTargetSourceFilename.FindOrAdd(InNewName.ToString()).SetSourceFilename(OldFilename);
+			const bool bOldIsGeneratedByEngine = MorphTargetImportedSourceFileInfo->IsGeneratedByEngine();
+			FMorphTargetImportedSourceFileInfo& NewData = LODInfoEntry.ImportedMorphTargetSourceFilename.FindOrAdd(InNewName.ToString());
+			NewData.SetSourceFilename(OldFilename);
+			NewData.SetGeneratedByEngine(bOldIsGeneratedByEngine);
 			LODInfoEntry.ImportedMorphTargetSourceFilename.Remove(InOldName.ToString());
 		}
 	}
diff --git a/Engine/Source/Runtime/Engine/Private/SkinnedAssetCommon.cpp b/Engine/Source/Runtime/Engine/Private/SkinnedAssetCommon.cpp
--- a/Engine/Source/Runtime/Engine/Private/SkinnedAssetCommon.cpp
+++ b/Engine/Source/Runtime/Engine/Private/SkinnedAssetCommon.cpp
@@ -14,6 +14,7 @@
 #include "SkeletalMeshLegacyCustomVersions.h"
 #include "UObject/CoreObjectVersion.h"
 #include "UObject/EditorObjectVersion.h"
+#include "UObject/FortniteSeasonBranchObjectVersion.h"
 #include "UObject/RenderingObjectVersion.h"
 
 #include UE_INLINE_GENERATED_CPP_BY_NAME(SkinnedAssetCommon)
@@ -23,11 +24,35 @@
 	FSkeletalMeshLODInfo
 -----------------------------------------------------------------------------*/
 
-#if WITH_EDITORONLY_DATA
+namespace UE::MorphTargetSourceHelper
+{
+	const FString& GetGeneratedByEnginePrefix()
+	{
+		static FString GeneratedByEngineStr = TEXT("GeneratedByEngine");
+		return GeneratedByEngineStr;
+	}
+}
 
 const FString& FMorphTargetImportedSourceFileInfo::GetSourceFilename() const
 {
-	return SourceFilename;
+	static FString EmptyString = FString();
+	if(!IsGeneratedByEngine())
+	{
+		return SourceFilename;
+	}
+	return EmptyString;
+}
+
+bool FMorphTargetImportedSourceFileInfo::IsGeneratedByEngine() const
+{
+	return SourceFilename.StartsWith(UE::MorphTargetSourceHelper::GetGeneratedByEnginePrefix());
+}
+
+void FMorphTargetImportedSourceFileInfo::SetGeneratedByEngine(bool bInGeneratedByEngine)
+{
+	FString NewSourceFilename = UE::MorphTargetSourceHelper::GetGeneratedByEnginePrefix();
+	NewSourceFilename += TEXT("_") + FGuid::NewGuid().ToString(EGuidFormats::Base36Encoded);
+	SetSourceFilename(NewSourceFilename);
 }
 
 void FMorphTargetImportedSourceFileInfo::SetSourceFilename(const FString& Filename)
@@ -57,6 +82,8 @@ FArchive& operator<<(FArchive& Ar, FMorphTargetImportedSourceFileInfo& MorphTarg
 	return Ar;
 }
 
+#if WITH_EDITORONLY_DATA
+
 static void SerializeReductionSettingsForDDC(FArchive& Ar, FSkeletalMeshOptimizationSettings& ReductionSettings)
 {
 	check(Ar.IsSaving());
