Create a script version of skeleton::GetCurveMetaDataNames function.

#rnx

diff --git a/Engine/Source/Editor/SkeletalMeshEditor/Private/SkeletalMeshEditorSubsystem.cpp b/Engine/Source/Editor/SkeletalMeshEditor/Private/SkeletalMeshEditorSubsystem.cpp
--- a/Engine/Source/Editor/SkeletalMeshEditor/Private/SkeletalMeshEditorSubsystem.cpp
+++ b/Engine/Source/Editor/SkeletalMeshEditor/Private/SkeletalMeshEditorSubsystem.cpp
@@ -1087,4 +1087,39 @@ void USkeletalMeshEditorSubsystem::SetMorphTargetsToGeneratedByEngineForAllSkele
 	}
 }
 
+void USkeletalMeshEditorSubsystem::GetSkeletonCurveMetaDataNames(const USkeleton* TargetSkeleton, TArray<FName>& OutNames, ESkelSubSysQueryCurvesMetatdataNamesFilter Filter)
+{
+	if (!TargetSkeleton)
+	{
+		return;
+	}
+	TargetSkeleton->ForEachCurveMetaData([&OutNames, &Filter](FName CurveName, const FCurveMetaData& CurveMetaData)
+	{
+		switch(Filter)
+		{
+			case ESkelSubSysQueryCurvesMetatdataNamesFilter::All:
+			{
+				OutNames.Add(CurveName);
+				break;
+			}
+			case ESkelSubSysQueryCurvesMetatdataNamesFilter::MorphTargetOnly:
+			{
+				if (CurveMetaData.Type.bMorphtarget)
+				{
+					OutNames.Add(CurveName);
+				}
+				break;
+			}
+			case ESkelSubSysQueryCurvesMetatdataNamesFilter::MaterialOnly:
+			{
+				if (CurveMetaData.Type.bMaterial)
+				{
+					OutNames.Add(CurveName);
+				}
+				break;
+			}
+		}
+	});
+}
+
 #undef LOCTEXT_NAMESPACE
diff --git a/Engine/Source/Editor/SkeletalMeshEditor/Public/SkeletalMeshEditorSubsystem.h b/Engine/Source/Editor/SkeletalMeshEditor/Public/SkeletalMeshEditorSubsystem.h
--- a/Engine/Source/Editor/SkeletalMeshEditor/Public/SkeletalMeshEditorSubsystem.h
+++ b/Engine/Source/Editor/SkeletalMeshEditor/Public/SkeletalMeshEditorSubsystem.h
@@ -11,6 +11,22 @@
 
 DECLARE_LOG_CATEGORY_EXTERN(LogSkeletalMeshEditorSubsystem, Log, All);
 
+class USkeleton;
+
+/**
+ * Blueprint enum to query skeleton curves
+ */
+UENUM(BlueprintType)
+enum class ESkelSubSysQueryCurvesMetatdataNamesFilter : uint8
+{
+	/** Query all metadata curves */
+	All,
+	/** Query only morph target metadata curves */
+	MorphTargetOnly,
+	/** Query only material metadata curves */
+	MaterialOnly
+};
+
 /**
 * USkeletalMeshEditorSubsystem
 * Subsystem for exposing skeletal mesh functionality to scripts
@@ -284,4 +300,13 @@ public:
 	 */
 	UFUNCTION(BlueprintCallable, Category = "Editor Scripting | SkeletalMesh")
 	static void SetMorphTargetsToGeneratedByEngineForAllSkeletalMesh(const TArray<FString>& OptionalNames, const TArray<FName>& OptionalPaths);
+
+	/**
+	 * Retrieve all curve metadata name from the skeleton.
+	 * 
+	 * @param TargetSkeleton The skeleton we query the data from.
+	 * @param OutNames - Receive all the name of the curves that has metadata in the skeleton.
+	 */
+	UFUNCTION(BlueprintCallable, Category = "Editor Scripting | Skeleton")
+	static void GetSkeletonCurveMetaDataNames(const USkeleton* TargetSkeleton, TArray<FName>& OutNames, ESkelSubSysQueryCurvesMetatdataNamesFilter Filter);
 };
