Skeletal Mesh: Make sure we can rename a socket that only exists on the skeletal mesh

#rnx

diff --git a/Engine/Source/Editor/SkeletalMeshEditor/Private/SkeletalMeshEditorSubsystem.cpp b/Engine/Source/Editor/SkeletalMeshEditor/Private/SkeletalMeshEditorSubsystem.cpp
--- a/Engine/Source/Editor/SkeletalMeshEditor/Private/SkeletalMeshEditorSubsystem.cpp
+++ b/Engine/Source/Editor/SkeletalMeshEditor/Private/SkeletalMeshEditorSubsystem.cpp
@@ -546,10 +546,11 @@ bool USkeletalMeshEditorSubsystem::RenameSocket(USkeletalMesh* SkeletalMesh, FNa
 	}
 
 	USkeletalMeshSocket* SkeletonSocket = SkeletalMesh->GetSkeleton()->FindSocket(OldName);
-	if (SkeletonSocket == nullptr)
+	// If they're one and the same socket, it's because the USkeletalMesh::FindSocket falls back on searching the skeleton if it can't find the
+	// socket on the mesh itself. 
+	if (SkeletonSocket == MeshSocket)
 	{
-		UE_LOG(LogSkeletalMeshEditorSubsystem, Error, TEXT("RenameSocket: The socket named '%s' does not exist on the Skeleton."), *OldName.ToString());
-		return false;
+		SkeletonSocket = nullptr;
 	}
 
 	FScopedTransaction Transaction(LOCTEXT("RenameSocket", "Rename Socket"));
@@ -557,9 +558,12 @@ bool USkeletalMeshEditorSubsystem::RenameSocket(USkeletalMesh* SkeletalMesh, FNa
 	MeshSocket->Modify();
 	MeshSocket->SocketName = NewName;
 
-	SkeletonSocket->SetFlags(RF_Transactional);
-	SkeletonSocket->Modify();
-	SkeletonSocket->SocketName = NewName;
+	if (SkeletonSocket)
+	{
+		SkeletonSocket->SetFlags(RF_Transactional);
+		SkeletonSocket->Modify();
+		SkeletonSocket->SocketName = NewName;
+	}
 
 	FPreviewAssetAttachContainer& PreviewAssetAttachContainer = SkeletalMesh->GetPreviewAttachedAssetContainer();
 	bool bMeshModified = false;
