PR #13830: Changed comment to make its meaning clear

diff --git a/Engine/Source/Editor/SkeletalMeshEditor/Private/SkeletalMeshEditor.cpp b/Engine/Source/Editor/SkeletalMeshEditor/Private/SkeletalMeshEditor.cpp
--- a/Engine/Source/Editor/SkeletalMeshEditor/Private/SkeletalMeshEditor.cpp
+++ b/Engine/Source/Editor/SkeletalMeshEditor/Private/SkeletalMeshEditor.cpp
@@ -191,7 +191,7 @@ bool FSkeletalMeshEditor::OnRequestClose(EAssetEditorCloseReason InCloseReason)
 
 		if (bHaveModifiedLOD && InCloseReason != EAssetEditorCloseReason::AssetForceDeleted)
 		{
-			// find out the user wants to do with this dirty material
+			// find out what the user wants to do with this dirty material
 			EAppReturnType::Type OkCancelReply = FMessageDialog::Open(
 				EAppMsgType::OkCancel,
 				FText::Format(LOCTEXT("SkeletalMeshEditorShouldApplyLODChanges", "We have to apply level of detail changes to {0} before exiting the skeletal mesh editor."), FText::FromString(PersonaToolkit->GetMesh()->GetName()))
