- Add BP CPU Access setter for Skeletal Meshes
- This mirrors functionality already available for static meshes

#Synced-CL 48563820

diff --git a/Engine/Source/Editor/SkeletalMeshEditor/Private/SkeletalMeshEditorSubsystem.cpp b/Engine/Source/Editor/SkeletalMeshEditor/Private/SkeletalMeshEditorSubsystem.cpp
--- a/Engine/Source/Editor/SkeletalMeshEditor/Private/SkeletalMeshEditorSubsystem.cpp
+++ b/Engine/Source/Editor/SkeletalMeshEditor/Private/SkeletalMeshEditorSubsystem.cpp
@@ -1191,6 +1191,31 @@ void USkeletalMeshEditorSubsystem::GetMorphTargetsGeneratedByEngine(USkeletalMes
 	}
 }
 
+void USkeletalMeshEditorSubsystem::SetAllowCPUAccess(USkeletalMesh* TargetMesh, bool bAllowCPUAccess)
+{
+	TGuardValue<bool> UnattendedScriptGuard(GIsRunningUnattendedScript, true);
+
+	if (!EditorScriptingHelpers::CheckIfInEditorAndPIE())
+	{
+		return;
+	}
+
+	if (TargetMesh == nullptr)
+	{
+		UE_LOG(LogSkeletalMeshEditorSubsystem, Error, TEXT("SetAllowCPUAccess: The SkeletalMesh is null."));
+		return;
+	}
+
+	TargetMesh->Modify();
+
+	for (int32 i=0; i < TargetMesh->GetLODNum(); ++i)
+	{
+		TargetMesh->GetLODInfo(i)->bAllowCPUAccess = bAllowCPUAccess;
+	}
+
+	TargetMesh->PostEditChange();
+}
+
 void USkeletalMeshEditorSubsystem::GetSkeletonCurveMetaDataNames(const USkeleton* TargetSkeleton, TArray<FName>& OutNames, ESkelSubSysQueryCurvesMetatdataNamesFilter Filter)
 {
 	if (!TargetSkeleton)
diff --git a/Engine/Source/Editor/SkeletalMeshEditor/Public/SkeletalMeshEditorSubsystem.h b/Engine/Source/Editor/SkeletalMeshEditor/Public/SkeletalMeshEditorSubsystem.h
--- a/Engine/Source/Editor/SkeletalMeshEditor/Public/SkeletalMeshEditorSubsystem.h
+++ b/Engine/Source/Editor/SkeletalMeshEditor/Public/SkeletalMeshEditorSubsystem.h
@@ -351,7 +351,11 @@ public:
 	*/
 	UFUNCTION(BlueprintCallable, Category = "Editor Scripting | SkeletalMesh")
 	static UE_API void GetMorphTargetsGeneratedByEngine(USkeletalMesh* TargetMesh, TArray<FName>& OutNames);
-	
+
+	/** Sets Skeletal Mesh Flag bAllowCPUAccess on all LOD levels. */
+	UFUNCTION(BlueprintCallable, Category = "Editor Scripting | SkeletalMesh")
+	static UE_API void SetAllowCPUAccess(USkeletalMesh* TargetMesh, bool bAllowCPUAccess);
+
 	/**
 	 * Retrieve all curve metadata name from the skeleton.
 	 * 
