Added RemoveSocket and bone parent/child accessor functions to SkeletalMeshSubsystem.

#Synced-CL 51698457

diff --git a/Engine/Source/Editor/SkeletalMeshEditor/Private/SkeletalMeshEditorSubsystem.cpp b/Engine/Source/Editor/SkeletalMeshEditor/Private/SkeletalMeshEditorSubsystem.cpp
--- a/Engine/Source/Editor/SkeletalMeshEditor/Private/SkeletalMeshEditorSubsystem.cpp
+++ b/Engine/Source/Editor/SkeletalMeshEditor/Private/SkeletalMeshEditorSubsystem.cpp
@@ -693,6 +693,157 @@ bool USkeletalMeshEditorSubsystem::RenameSocket(USkeletalMesh* SkeletalMesh, FNa
 	return true;
 }
 
+bool USkeletalMeshEditorSubsystem::RemoveSocket(USkeletalMesh* SkeletalMesh, FName SocketName)
+{
+	if (SkeletalMesh == nullptr)
+	{
+		UE_LOG(LogSkeletalMeshEditorSubsystem, Error, TEXT("RemoveSocket: The SkeletalMesh is null."));
+		return false;
+	}
+
+	if (SkeletalMesh->GetSkeleton() == nullptr)
+	{
+		UE_LOG(LogSkeletalMeshEditorSubsystem, Error, TEXT("RemoveSocket: The SkeletalMesh's Skeleton is null."));
+		return false;
+	}
+
+	if (SocketName == NAME_None)
+	{
+		UE_LOG(LogSkeletalMeshEditorSubsystem, Error, TEXT("RemoveSocket: The SocketName is None."));
+		return false;
+	}
+
+	USkeletalMeshSocket* FoundSocket = SkeletalMesh->FindSocket(SocketName);
+	if (FoundSocket == nullptr)
+	{
+		UE_LOG(LogSkeletalMeshEditorSubsystem, Error, TEXT("RemoveSocket: The socket named '%s' does not exist on the SkeletalMesh."), *SocketName.ToString());
+		return false;
+	}
+
+	USkeletalMeshSocket* SkeletonSocket = SkeletalMesh->GetSkeleton()->FindSocket(SocketName);
+	// If they're the same pointer, FindSocket fell back to the skeleton — there is no mesh-specific socket.
+	if (SkeletonSocket == FoundSocket)
+	{
+		SkeletonSocket = nullptr;
+	}
+
+	FScopedTransaction Transaction(LOCTEXT("RemoveSocket", "Remove Socket"));
+
+	TArray<TObjectPtr<USkeletalMeshSocket>>& MeshSocketList = SkeletalMesh->GetMeshOnlySocketList();
+	if (MeshSocketList.Contains(FoundSocket))
+	{
+		FoundSocket->SetFlags(RF_Transactional);
+		FoundSocket->Modify();
+		SkeletalMesh->Modify();
+		MeshSocketList.Remove(FoundSocket);
+	}
+	else
+	{
+		// Socket exists only on the skeleton (FindSocket fell through to it).
+		FoundSocket->SetFlags(RF_Transactional);
+		FoundSocket->Modify();
+		SkeletalMesh->GetSkeleton()->Modify();
+		SkeletalMesh->GetSkeleton()->Sockets.Remove(FoundSocket);
+	}
+
+	if (SkeletonSocket != nullptr)
+	{
+		SkeletonSocket->SetFlags(RF_Transactional);
+		SkeletonSocket->Modify();
+		SkeletalMesh->GetSkeleton()->Modify();
+		SkeletalMesh->GetSkeleton()->Sockets.Remove(SkeletonSocket);
+	}
+
+	// Clean up any preview-attached assets referencing the removed socket on the mesh.
+	bool bMeshModified = false;
+	FPreviewAssetAttachContainer& MeshPreviewContainer = SkeletalMesh->GetPreviewAttachedAssetContainer();
+	for (int32 Index = MeshPreviewContainer.Num() - 1; Index >= 0; --Index)
+	{
+		if (MeshPreviewContainer[Index].AttachedTo == SocketName)
+		{
+			if (!bMeshModified)
+			{
+				SkeletalMesh->Modify();
+				bMeshModified = true;
+			}
+			MeshPreviewContainer.RemoveAtSwap(Index);
+		}
+	}
+
+	// Clean up any preview-attached assets referencing the removed socket on the skeleton.
+	bool bSkeletonModified = false;
+	FPreviewAssetAttachContainer& SkeletonPreviewContainer = SkeletalMesh->GetSkeleton()->PreviewAttachedAssetContainer;
+	for (int32 Index = SkeletonPreviewContainer.Num() - 1; Index >= 0; --Index)
+	{
+		if (SkeletonPreviewContainer[Index].AttachedTo == SocketName)
+		{
+			if (!bSkeletonModified)
+			{
+				SkeletalMesh->GetSkeleton()->Modify();
+				bSkeletonModified = true;
+			}
+			SkeletonPreviewContainer.RemoveAtSwap(Index);
+		}
+	}
+
+	return true;
+}
+
+FName USkeletalMeshEditorSubsystem::GetBoneParent(const USkeletalMesh* SkeletalMesh, FName BoneName)
+{
+	if (!SkeletalMesh)
+	{
+		UE_LOG(LogSkeletalMeshEditorSubsystem, Error, TEXT("GetBoneParent: The SkeletalMesh is null."));
+		return NAME_None;
+	}
+
+	const FReferenceSkeleton& RefSkel = SkeletalMesh->GetRefSkeleton();
+	const int32 BoneIndex = RefSkel.FindBoneIndex(BoneName);
+	if (BoneIndex == INDEX_NONE)
+	{
+		UE_LOG(LogSkeletalMeshEditorSubsystem, Error, TEXT("GetBoneParent: Bone '%s' not found."), *BoneName.ToString());
+		return NAME_None;
+	}
+
+	const int32 ParentIndex = RefSkel.GetParentIndex(BoneIndex);
+	if (ParentIndex == INDEX_NONE)
+	{
+		return NAME_None;
+	}
+
+	return RefSkel.GetBoneName(ParentIndex);
+}
+
+TArray<FName> USkeletalMeshEditorSubsystem::GetBoneChildren(const USkeletalMesh* SkeletalMesh, FName BoneName)
+{
+	TArray<FName> Children;
+
+	if (!SkeletalMesh)
+	{
+		UE_LOG(LogSkeletalMeshEditorSubsystem, Error, TEXT("GetBoneChildren: The SkeletalMesh is null."));
+		return Children;
+	}
+
+	const FReferenceSkeleton& RefSkel = SkeletalMesh->GetRefSkeleton();
+	const int32 BoneIndex = RefSkel.FindBoneIndex(BoneName);
+	if (BoneIndex == INDEX_NONE)
+	{
+		UE_LOG(LogSkeletalMeshEditorSubsystem, Error, TEXT("GetBoneChildren: Bone '%s' not found."), *BoneName.ToString());
+		return Children;
+	}
+
+	const int32 NumBones = RefSkel.GetNum();
+	for (int32 i = 0; i < NumBones; ++i)
+	{
+		if (RefSkel.GetParentIndex(i) == BoneIndex)
+		{
+			Children.Add(RefSkel.GetBoneName(i));
+		}
+	}
+
+	return Children;
+}
+
 int32 USkeletalMeshEditorSubsystem::GetLODCount(USkeletalMesh* SkeletalMesh)
 {
 	if (SkeletalMesh == nullptr)
diff --git a/Engine/Source/Editor/SkeletalMeshEditor/Public/SkeletalMeshEditorSubsystem.h b/Engine/Source/Editor/SkeletalMeshEditor/Public/SkeletalMeshEditorSubsystem.h
--- a/Engine/Source/Editor/SkeletalMeshEditor/Public/SkeletalMeshEditorSubsystem.h
+++ b/Engine/Source/Editor/SkeletalMeshEditor/Public/SkeletalMeshEditorSubsystem.h
@@ -217,6 +217,30 @@ public:
 	UFUNCTION(BlueprintCallable, Category = "Skeletal Mesh Utilities", meta = (ScriptMethod))
 	static UE_API bool RenameSocket(USkeletalMesh* SkeletalMesh, FName OldName, FName NewName);
 
+	/** Remove a socket from a skeletal mesh
+	 * @param SkeletalMesh	The mesh from which we are removing a socket
+	 * @param SocketName    The name of the socket to remove
+	 * @return true if the removal succeeded.
+	 */
+	UFUNCTION(BlueprintCallable, Category = "Skeletal Mesh Utilities", meta = (ScriptMethod))
+	static UE_API bool RemoveSocket(USkeletalMesh* SkeletalMesh, FName SocketName);
+
+	/** Get the parent bone name of a bone in a skeletal mesh.
+	 * @param SkeletalMesh  The mesh to query.
+	 * @param BoneName      The name of the bone to query.
+	 * @return The name of the parent bone, or NAME_None if the bone is the root or does not exist.
+	 */
+	UFUNCTION(BlueprintCallable, Category = "Skeletal Mesh Utilities", meta = (ScriptMethod))
+	static UE_API FName GetBoneParent(const USkeletalMesh* SkeletalMesh, FName BoneName);
+
+	/** Get the direct children of a bone in a skeletal mesh.
+	 * @param SkeletalMesh  The mesh to query.
+	 * @param BoneName      The name of the bone to query.
+	 * @return An array of bone names that are direct children of the given bone.
+	 */
+	UFUNCTION(BlueprintCallable, Category = "Skeletal Mesh Utilities", meta = (ScriptMethod))
+	static UE_API TArray<FName> GetBoneChildren(const USkeletalMesh* SkeletalMesh, FName BoneName);
+
 	/**
 	 * Retrieve the number of LOD contain in the specified skeletal mesh.
 	 *
