SkeletalMeshEditor selection fix

#Synced-CL 52103551

diff --git a/Engine/Source/Editor/SkeletalMeshEditor/Private/SkeletalMeshEditor.cpp b/Engine/Source/Editor/SkeletalMeshEditor/Private/SkeletalMeshEditor.cpp
--- a/Engine/Source/Editor/SkeletalMeshEditor/Private/SkeletalMeshEditor.cpp
+++ b/Engine/Source/Editor/SkeletalMeshEditor/Private/SkeletalMeshEditor.cpp
@@ -1279,15 +1279,6 @@ void FSkeletalMeshEditor::HandleObjectsSelected(const TArray<UObject*>& InObject
 	}
 }
 
-void FSkeletalMeshEditor::OnSelectionSetChanged(const UTypedElementSelectionSet* InSelectionSet)
-{
-	TArray<UObject*> SelectedObjects = InSelectionSet->GetSelectedObjects();
-	if (DetailsView.IsValid())
-	{
-		DetailsView->SetObjects(SelectedObjects);
-	}
-}
-
 void FSkeletalMeshEditor::HandleObjectSelected(UObject* InObject)
 {
 	if (DetailsView.IsValid())
diff --git a/Engine/Source/Editor/SkeletalMeshEditor/Private/SkeletalMeshEditor.h b/Engine/Source/Editor/SkeletalMeshEditor/Private/SkeletalMeshEditor.h
--- a/Engine/Source/Editor/SkeletalMeshEditor/Private/SkeletalMeshEditor.h
+++ b/Engine/Source/Editor/SkeletalMeshEditor/Private/SkeletalMeshEditor.h
@@ -128,8 +128,6 @@ public:
 
 	void HandleObjectsSelected(const TArray<UObject*>& InObjects);
 
-	virtual void OnSelectionSetChanged(const UTypedElementSelectionSet* InSelectionSet) override;
-
 	// Returns the currently hosted toolkit. Can be invalid if no toolkit is being hosted.
 	TSharedPtr<IToolkit> GetHostedToolkit() const { return HostedToolkit; }
 	
