Skeletal Mesh Editor: enable default marquee selection
- this was already implemented but the ITF/gizmo integration was partial at that time

#rnx

#Synced-CL 52515883

diff --git a/Engine/Source/Editor/SkeletalMeshEditor/Private/SkeletalMeshEditor.cpp b/Engine/Source/Editor/SkeletalMeshEditor/Private/SkeletalMeshEditor.cpp
--- a/Engine/Source/Editor/SkeletalMeshEditor/Private/SkeletalMeshEditor.cpp
+++ b/Engine/Source/Editor/SkeletalMeshEditor/Private/SkeletalMeshEditor.cpp
@@ -68,6 +68,7 @@
 #include "InterchangeAssetImportData.h"
 #include "InterchangeManager.h"
 #include "SkeletalMeshEditorModeUILayer.h"
+#include "EditorDragTools/FrustumSelectInteraction.h"
 #include "Tools/EdModeInteractiveToolsContext.h"
 #include "ViewportInteractions/ViewportInteractionsBehaviorSource.h"
 #include "ViewportInteractions/ViewportInteractionsUtilities.h"
@@ -310,6 +311,15 @@ void FSkeletalMeshEditor::InitSkeletalMeshEditor(const EToolkitMode::Type Mode,
 	}
 }
 
+void FSkeletalMeshEditor::OnBuildViewportInteractions(UViewportInteractionsBehaviorSource* Source)
+{
+	ISkeletalMeshEditor::OnBuildViewportInteractions(Source);
+	if (Source)
+	{
+		Source->AddInteraction<UFrustumSelectInteraction>();
+	}
+}
+
 void FSkeletalMeshEditor::OnToolkitHostingStarted(const TSharedRef<IToolkit>& Toolkit)
 {
 	ISkeletalMeshEditor::OnToolkitHostingStarted(Toolkit);
diff --git a/Engine/Source/Editor/SkeletalMeshEditor/Private/SkeletalMeshEditor.h b/Engine/Source/Editor/SkeletalMeshEditor/Private/SkeletalMeshEditor.h
--- a/Engine/Source/Editor/SkeletalMeshEditor/Private/SkeletalMeshEditor.h
+++ b/Engine/Source/Editor/SkeletalMeshEditor/Private/SkeletalMeshEditor.h
@@ -137,6 +137,9 @@ public:
 	virtual TSharedPtr<FUICommandInfo> GetResetBoneTransformsCommand() override;
 	virtual TSharedPtr<FUICommandInfo> GetResetAllBonesTransformsCommand() override;
 	
+protected:
+	void OnBuildViewportInteractions(UViewportInteractionsBehaviorSource* Source) override;
+	
 private:
 	void HandleObjectSelected(UObject* InObject);
 	
